#UE4 Platforms
#HTML5

   Integerate GDC Changes into main. tested w/ QAGame, Soul and the moz custom built firefox.

   - resizable HTML5 canvas.
   - fix mouse coordinates for html5
   - depth fade using alpha ( [REDACTED_HANDLE] - use real webgl depth  textures )
   - fix landscape rendering. ( bgra is not supported on webgl )
   - expose a console command handler to javascript.  ( used while reszing )
   - update html5 .html5 game template.

diff --git a/Engine/Build/HTML5/Game.html.template b/Engine/Build/HTML5/Game.html.template
--- a/Engine/Build/HTML5/Game.html.template
+++ b/Engine/Build/HTML5/Game.html.template
@@ -4,80 +4,203 @@
     <meta charset="utf-8">
     <meta http-equiv="Content-Type" content="text/html; charset=utf-8">
     <title>%GAME%</title>
-    <style>
-      .emscripten { padding-right: 0; margin-left: auto; margin-right: auto; display: block; }
-      canvas.emscripten { border: 1px solid black; }
-      textarea.emscripten { font-family: monospace; width: 80%; }
-      div.emscripten { text-align: center; }
+    <style type="text/css">
+body {
+  background: black;
+  color: gray;
+  font-family: Helvetica, sans-serif;
+}
+
+.emscripten {
+  padding-right: 0;
+  margin-left: auto;
+  margin-right: auto;
+  display: block;
+}
+
+div#status {
+  font-size: 70%;
+}
+
+textarea.emscripten {
+  font-family: consolas, "lucida sans console", monospace;
+  background: gray;
+  font-size: 80%;
+  width: 80%;
+  display: none;
+}
+
+div.emscripten {
+  text-align: center;
+}
+
+div#display_wrapper {
+  position: relative;
+  margin: 1em auto 10px auto;
+  text-align: center;
+  min-width: 640px;
+  max-width: 95%;
+  /*border: 1px solid red;*/
+}
+
+canvas.main {
+  border: none;
+  background: #444;
+  display: block;
+}
+
+canvas.main:not([fullscreen]) {
+  padding-right: 0;
+  margin-left: auto;
+  margin-right: auto;
+  width: 100%;
+}
+
+input[type=button] {
+  background: #444;
+  border: 2px solid #444;
+  border-radius: 2px;
+}
+
+input[type=button]:hover {
+  background: #777;
+  border-color: #777;
+}
+
     </style>
   </head>
   <body>
-    <hr/>
     <div class="emscripten" id="status">Downloading...</div>
     <div class="emscripten">
       <progress value="0" max="100" id="progress" hidden=1></progress>  
     </div>
-    <canvas class="emscripten" id="canvas" oncontextmenu="event.preventDefault()"></canvas>
-    <hr/>
+
+    <!-- Canvas resolution needs to be divisible by 8; default to 904x1600 for in-page resolution -->
+    <div id="display_wrapper">
+      <canvas class="main" id="canvas" oncontextmenu="event.preventDefault()" height="904" width="1600"></canvas>
+    </div>
+
     <div class="emscripten">
-		<input type="button" value="fullscreen" onclick="Module.requestFullScreen()">
-		<input type="button" value="Pause" onclick="Module['pauseMainLoop']();">
-		<input type="button" value="Resume" onclick="Module['resumeMainLoop']();">
-		<input type="button" value="Quit" onclick="__exit(0)">
-	</div>
-    <hr/>
+      <input type="button" value="Fullscreen" onclick="Module.requestFullScreen(false)">
+      <input type="button" value="Pause" onclick="Module['pauseMainLoop']();">
+      <input type="button" value="Resume" onclick="Module['resumeMainLoop']();">
+      <input type="button" value="Quit" onclick="__exit(0)">
+    </div>
+
     <textarea class="emscripten" id="output" rows="8"></textarea>
-    <hr>
-    <script type='text/javascript'>
-	 var tstart = Date.now();
- 
-      // connect to canvas
-      var Module = {
-        preRun: [],
-        postRun: [],
-		noImageDecoding: true,
-		noAudioDecoding: true,
+
+    <script type="text/javascript">
+
+var UE4 = {
+  get console_command() {
+    var fn = Module.cwrap('execute_console_command', null, ['string']);
+    delete UE4["console_command"];
+    UE4.console_command = fn;
+    return fn;
+  }
+};
+
+function resizeListener(width, height) {
+  console.log("resizeListener:", width, height);
+  UE4.console_command("r.SetRes " + width + "x" + height);
+}
+
+var gWasFullScreen = false;
+var gOldWidth, gOldHeight;
+
+function fullScreenChange() {
+  var fsElem = document['fullScreenElement'] ||
+               document['mozFullScreenElement'] ||
+               document['webkitFullScreenElement'] ||
+               document['msFullScreenElement'] ||
+               document['webkitCurrentFullScreenElement'];
+  var canvas = document.getElementById("canvas");
+
+  if (fsElem == canvas || fsElem == canvas.parentNode) {
+    // we just went full screen
+    gWasFullScreen = true;
+
+    // update the canvas size
+    gOldWidth = canvas.width;
+    gOldHeight = canvas.height;
+
+    var width = window.screen.width;
+    var height = window.screen.height;
+
+    canvas.width = width;
+    canvas.height = height;
+    canvas.setAttribute("fullscreen", "true");
+
+    console.log("r.SetRes", canvas.width, canvas.height);
+    UE4.console_command("r.SetRes " + canvas.width + "x" + canvas.height);
+  } else if (gWasFullScreen) {
+    // we just left full screen
+    gWasFullScreen = false;
+
+    canvas.width = gOldWidth;
+    canvas.height = gOldHeight;
+    canvas.removeAttribute("fullscreen");
+
+    console.log("r.SetRes", canvas.width, canvas.height);
+    UE4.console_command("r.SetRes " + canvas.width + "x" + canvas.height);
+  }
+}
+
+function preRunHandler() {
+  document.addEventListener('fullscreenchange', fullScreenChange, false);
+  document.addEventListener('mozfullscreenchange', fullScreenChange, false);
+  document.addEventListener('webkitfullscreenchange', fullScreenChange, false);
+  document.addEventListener('MSFullscreenChange', fullScreenChange, false);
+
+  // this needs to die
+  Browser.updateCanvasDimensions = function() { };
+  
+}
+
+var Module = {
+  preRun: [preRunHandler],
+  postRun: [],
+  noImageDecoding: true,
+  noAudioDecoding: true,
         arguments: [%UE4CMDLINE%],
-        print: (function() {
-          var element = document.getElementById('output');
-          element.value = ''; // clear browser cache
-          return function(text) {
-				element.value += text + "\n";
-				element.scrollTop = 99999; // focus on bottom
-		  };
-        })(),
-        printErr: function(text) {
-          if (0) { // XXX disabled for safety typeof dump == 'function') {
-            dump(text + '\n'); // fast, straight to the real console
-          } else {
-            console.log(text);
-          }
-        },
-        canvas: document.getElementById('canvas'),
-        setStatus: function(text) {
-          if (Module.setStatus.interval) clearInterval(Module.setStatus.interval);
-          var m = text.match(/([^(]+)\((\d+(\.\d+)?)\/(\d+)\)/);
-          var statusElement = document.getElementById('status');
-          var progressElement = document.getElementById('progress');
-          if (m) {
-            text = m[1];
-            progressElement.value = parseInt(m[2])*100;
-            progressElement.max = parseInt(m[4])*100;
-            progressElement.hidden = false;
-          } else {
-            progressElement.value = null;
-            progressElement.max = null;
-            progressElement.hidden = true;
-          }
-          statusElement.innerHTML = text;
-        },
-        totalDependencies: 0,
-        monitorRunDependencies: function(left) {
-          this.totalDependencies = Math.max(this.totalDependencies, left);
-          Module.setStatus(left ? 'Preparing... (' + (this.totalDependencies-left) + '/' + this.totalDependencies + ')' : 'All downloads complete.');
-        }
-      };
-      Module.setStatus('Downloading...');
+  print: (function() {
+    var element = document.getElementById('output');
+    element.value = ''; // clear browser cache
+    return function(text) {
+      // show it if it's hidden first
+      element.style.display = "block";
+      element.value += text + "\n";
+      element.scrollTop = 99999; // focus on bottom
+    };
+  })(),
+  printErr: function(text) {
+    console.log(text);
+  },
+  canvas: document.getElementById('canvas'),
+  setStatus: function(text) {
+    if (Module.setStatus.interval) clearInterval(Module.setStatus.interval);
+    var m = text.match(/([^(]+)\((\d+(\.\d+)?)\/(\d+)\)/);
+    var statusElement = document.getElementById('status');
+    var progressElement = document.getElementById('progress');
+    if (m) {
+      text = m[1];
+      progressElement.value = parseInt(m[2])*100;
+      progressElement.max = parseInt(m[4])*100;
+      progressElement.hidden = false;
+    } else {
+      progressElement.value = null;
+      progressElement.max = null;
+      progressElement.hidden = true;
+    }
+    statusElement.innerHTML = text;
+  },
+  totalDependencies: 0,
+  monitorRunDependencies: function(left) {
+    this.totalDependencies = Math.max(this.totalDependencies, left);
+    Module.setStatus(left ? 'Preparing... (' + (this.totalDependencies-left) + '/' + this.totalDependencies + ')' : 'All downloads complete.');
+  }
+};
+Module.setStatus('Downloading...');
     </script>
 	<script type="text/javascript">
       var tmid = Date.now();
diff --git a/Engine/Plugins/Messaging/UdpMessaging/UdpMessaging.uplugin b/Engine/Plugins/Messaging/UdpMessaging/UdpMessaging.uplugin
--- a/Engine/Plugins/Messaging/UdpMessaging/UdpMessaging.uplugin
+++ b/Engine/Plugins/Messaging/UdpMessaging/UdpMessaging.uplugin
@@ -15,7 +15,8 @@
 		{
 			"Name" : "UdpMessaging",
 			"Type" : "Runtime",
-			"LoadingPhase" : "PreDefault"
+			"LoadingPhase" : "PreDefault",
+			"BlacklistPlatforms" : [ "HTML5" ]
 		}
 	],
 
diff --git a/Engine/Shaders/Common.usf b/Engine/Shaders/Common.usf
--- a/Engine/Shaders/Common.usf
+++ b/Engine/Shaders/Common.usf
@@ -582,9 +582,11 @@ float CalcSceneDepth( float2 ScreenUV )
 			#if IOS
 				// Only call FramebufferFetch when actually compiling for IOS ES2.
 				return FramebufferFetchES2().w;
-			#else
-				return 65000.0f;
-			#endif
+			#elif WEBGL
+				return Texture2DSampleLevel(SceneAlphaCopyTexture, SceneAlphaCopyTextureSampler, ScreenUV, 0).r;
+			#else 
+				return 65000.0f; 
+			#endif 
 		#else
 			return ConvertFromDeviceZ(Texture2DSampleLevel(SceneDepthTexture, SceneDepthTextureSampler, ScreenUV, 0).r);
 		#endif
diff --git a/Engine/Shaders/LandscapeVertexFactory.usf b/Engine/Shaders/LandscapeVertexFactory.usf
--- a/Engine/Shaders/LandscapeVertexFactory.usf
+++ b/Engine/Shaders/LandscapeVertexFactory.usf
@@ -177,7 +177,11 @@ float3x3 VertexFactoryGetPerPixelTangentBasis(FVertexFactoryInterpolantsVSToPS I
 #else
 	float4 SampleValue = Texture2DSample(NormalmapTexture, NormalmapTextureSampler, Interpolants.WeightMapTexCoord);
 #endif
-	float2 SampleNormal = float2(SampleValue.b, SampleValue.a) * float2(2.0,2.0) - float2(1.0,1.0);
+#if WEBGL
+	float2 SampleNormal = float2(SampleValue.r, SampleValue.a) * float2(2.0,2.0) - float2(1.0,1.0);
+#else
+    float2 SampleNormal = float2(SampleValue.b, SampleValue.a) * float2(2.0,2.0) - float2(1.0,1.0);
+#endif
 	float3 WorldNormal = float3( SampleNormal, sqrt(max(1.0-dot(SampleNormal,SampleNormal),0.0)) );
 
 	Result = CalcTangentBasisFromWorldNormal(WorldNormal);
@@ -326,7 +330,13 @@ FMaterialPixelParameters GetMaterialPixelParameters(FVertexFactoryInterpolantsVS
 #endif
 
 #if ES2_PROFILE && PIXELSHADER
+
+#if WEBGL
+	Result.LayerWeights.xy = Texture2DSample(NormalmapTexture, NormalmapTextureSampler, Interpolants.WeightMapTexCoord).bg;
+#else
 	Result.LayerWeights.xy = Texture2DSample(NormalmapTexture, NormalmapTextureSampler, Interpolants.WeightMapTexCoord).rg;
+#endif 
+ 
 	Result.LayerWeights.z = saturate(1.f - Result.LayerWeights.x - Result.LayerWeights.y);
 	Result.LayerWeights.w = 0; // Total weight need to be 1
 #endif
diff --git a/Engine/Source/Developer/HTML5/HTML5TargetPlatform/Private/HTML5TargetPlatform.cpp b/Engine/Source/Developer/HTML5/HTML5TargetPlatform/Private/HTML5TargetPlatform.cpp
--- a/Engine/Source/Developer/HTML5/HTML5TargetPlatform/Private/HTML5TargetPlatform.cpp
+++ b/Engine/Source/Developer/HTML5/HTML5TargetPlatform/Private/HTML5TargetPlatform.cpp
@@ -220,7 +220,22 @@ void FHTML5TargetPlatform::GetTextureFormats( const UTexture* Texture, TArray<FN
 	// Determine the pixel format of the (un/)compressed texture
 	if (bNoCompression)
 	{
-		TextureFormatName = NameBGRA8;
+			if (Texture->HasHDRSource())
+			{
+				TextureFormatName = NameBGRA8;
+			}
+			else if (SourceFormat == TSF_G8 || Texture->CompressionSettings == TC_Grayscale)
+			{
+				TextureFormatName = NameG8;
+			}
+			else if (Texture->LODGroup == TEXTUREGROUP_Shadowmap)
+			{
+				TextureFormatName = NameG8;
+			}
+			else
+			{
+				TextureFormatName = NameBGRA8;
+			}
 	}
 	else if (Texture->CompressionSettings == TC_HDR)
 	{
diff --git a/Engine/Source/Programs/UnrealBuildTool/HTML5/HTML5ToolChain.cs b/Engine/Source/Programs/UnrealBuildTool/HTML5/HTML5ToolChain.cs
--- a/Engine/Source/Programs/UnrealBuildTool/HTML5/HTML5ToolChain.cs
+++ b/Engine/Source/Programs/UnrealBuildTool/HTML5/HTML5ToolChain.cs
@@ -81,70 +81,71 @@ namespace UnrealBuildTool
 
 		static string GetSharedArguments_Global(CPPTargetConfiguration TargetConfiguration, string Architecture)
 		{
-			string Result = " ";
+            string Result = " ";
 
-			if (Architecture == "-win32")
-			{
+            if (Architecture == "-win32")
+            {
                 return Result;
             }
 
-			// 			Result += " -funsigned-char";
-			// 			Result += " -fno-strict-aliasing";
-			Result += " -fno-exceptions";
-			// 			Result += " -fno-short-enums";
+            // 			Result += " -funsigned-char";
+            // 			Result += " -fno-strict-aliasing";
+            Result += " -fno-exceptions";
+            // 			Result += " -fno-short-enums";
 
-			Result += " -Wno-unused-value"; // appErrorf triggers this
-			Result += " -Wno-switch"; // many unhandled cases
-			Result += " -Wno-tautological-constant-out-of-range-compare"; // disables some warnings about comparisons from TCHAR being a char
-			// this hides the "warning : comparison of unsigned expression < 0 is always false" type warnings due to constant comparisons, which are possible with template arguments
-			Result += " -Wno-tautological-compare";
+            Result += " -Wno-unused-value"; // appErrorf triggers this
+            Result += " -Wno-switch"; // many unhandled cases
+            Result += " -Wno-tautological-constant-out-of-range-compare"; // disables some warnings about comparisons from TCHAR being a char
+            // this hides the "warning : comparison of unsigned expression < 0 is always false" type warnings due to constant comparisons, which are possible with template arguments
+            Result += " -Wno-tautological-compare";
 
-			// okay, in UE4, we'd fix the code for these, but in UE3, not worth it
-			Result += " -Wno-logical-op-parentheses"; // appErrorf triggers this
-			Result += " -Wno-array-bounds"; // some VectorLoads go past the end of the array, but it's okay in that case
+            // okay, in UE4, we'd fix the code for these, but in UE3, not worth it
+            Result += " -Wno-logical-op-parentheses"; // appErrorf triggers this
+            Result += " -Wno-array-bounds"; // some VectorLoads go past the end of the array, but it's okay in that case
             Result += " -Wno-invalid-offsetof"; // too many warnings kills windows clang. 
             
 
-			// JavsScript option overrides (see src/settings.js)
+            // JavsScript option overrides (see src/settings.js)
 
-			// we have to specify the full amount of memory with Asm.JS (1.5 G)
+            // we have to specify the full amount of memory with Asm.JS (1.5 G)
             // I wonder if there's a per game way to change this. 
             Result += " -s TOTAL_MEMORY=1610612736";
 
-			// no need for exceptions
-			Result += " -s DISABLE_EXCEPTION_CATCHING=1";
-			// enable checking for missing functions at link time as opposed to runtime
-			Result += " -s WARN_ON_UNDEFINED_SYMBOLS=1";
-			// we want full ES2
-			Result += " -s FULL_ES2=1 ";
-			// don't need UTF8 string support, and it slows string ops down
-			Result += " -s UTF_STRING_SUPPORT=0";
-            Result += "  ";
-
-			if (TargetConfiguration == CPPTargetConfiguration.Debug)
-			{
-				Result += " -O0";
-			}
-			if (TargetConfiguration == CPPTargetConfiguration.Debug || TargetConfiguration == CPPTargetConfiguration.Development)
-			{
-				Result += " -s GL_ASSERTIONS=1 ";
-			}
-			if (TargetConfiguration == CPPTargetConfiguration.Development)
-			{
+            // no need for exceptions
+            Result += " -s DISABLE_EXCEPTION_CATCHING=1";
+            // enable checking for missing functions at link time as opposed to runtime
+            Result += " -s WARN_ON_UNDEFINED_SYMBOLS=1";
+            // we want full ES2
+            Result += " -s FULL_ES2=1 ";
+            // don't need UTF8 string support, and it slows string ops down
+            Result += " -s UTF_STRING_SUPPORT=0";
+            // export console command handler. Export main func too because default exports ( e.g Main ) are overridden if we use custom exported functions. 
+            Result += " -s EXPORTED_FUNCTIONS=\"['_main', '_execute_console_command']\" ";
+
+            if (TargetConfiguration == CPPTargetConfiguration.Debug)
+            {
+	            Result += " -O0";
+            }
+            if (TargetConfiguration == CPPTargetConfiguration.Debug || TargetConfiguration == CPPTargetConfiguration.Development)
+            {
+	            Result += " -s GL_ASSERTIONS=1 ";
+            }
+            if (TargetConfiguration == CPPTargetConfiguration.Development)
+            {
                 Result += " -O2 -s ASM_JS=1 -s OUTLINING_LIMIT=110000";
-			}
-			if (TargetConfiguration == CPPTargetConfiguration.Shipping)
-			{
-				Result += " -O2 -s ASM_JS=1 -s OUTLINING_LIMIT=110000";
-			}
+            }
+            if (TargetConfiguration == CPPTargetConfiguration.Shipping)
+            {
+	            Result += " -O2 -s ASM_JS=1 -s OUTLINING_LIMIT=110000";
+            }
 
-			// NOTE: This may slow down the compiler's startup time!
+            // NOTE: This may slow down the compiler's startup time!
             if ( !bEnableFastIteration )
-			{ 
+            { 
                 Result += " --memory-init-file 1";
-			}
+            }
 
-			return Result;
+            return Result;
 		}
 		
 		static string GetCLArguments_Global(CPPEnvironment CompileEnvironment)
diff --git a/Engine/Source/Runtime/Core/Private/HTML5/HTML5Application.cpp b/Engine/Source/Runtime/Core/Private/HTML5/HTML5Application.cpp
--- a/Engine/Source/Runtime/Core/Private/HTML5/HTML5Application.cpp
+++ b/Engine/Source/Runtime/Core/Private/HTML5/HTML5Application.cpp
@@ -13,6 +13,7 @@ FHTML5Application* FHTML5Application::CreateHTML5Application()
 FHTML5Application::FHTML5Application()
 	: GenericApplication( MakeShareable( new FHTML5Cursor() ) )
 	, InputInterface( FHTML5InputInterface::Create( MessageHandler, Cursor ) )
+	, ApplicationWindow(FHTML5Window::Make())
 {
 }
 
@@ -24,36 +25,40 @@ void FHTML5Application::SetMessageHandler( const TSharedRef< FGenericApplication
 
 void FHTML5Application::PollGameDeviceState( const float TimeDelta )
 {
-	// Poll game device state and send new events
-	InputInterface->Tick( TimeDelta );
+	SDL_Event Event;
+	while (SDL_PollEvent(&Event)) 
+	{
+		// Tick Input Interface. 
+		if ( Event.type == SDL_VIDEORESIZE )
+		{
+			SDL_ResizeEvent *r = (SDL_ResizeEvent*)&Event;
+			MessageHandler->OnSizeChanged(  ApplicationWindow, r->w,r->h); 
+		}
+		else 
+		{
+			InputInterface->Tick( TimeDelta, Event );
+		}
+	}
 	InputInterface->SendControllerEvents();
 }
 
 FPlatformRect FHTML5Application::GetWorkArea( const FPlatformRect& CurrentWindow ) const
 {
-	//@todo HTML5: Use the actual device settings here.
-	FPlatformRect WorkArea;
-	WorkArea.Left = 0;
-	WorkArea.Top = 0;
-	WorkArea.Right = 800;
-	WorkArea.Bottom = 600;
-
-	return WorkArea;
+	return  FHTML5Window::GetScreenRect();
 }
 
 void FHTML5Application::GetDisplayMetrics( FDisplayMetrics& OutDisplayMetrics ) const
 {
-	//@todo HTML5: Use the actual device settings here.
-	OutDisplayMetrics.PrimaryDisplayWidth = 800;
-	OutDisplayMetrics.PrimaryDisplayHeight = 600;
-
-	OutDisplayMetrics.PrimaryDisplayWorkAreaRect.Left = 0;
-	OutDisplayMetrics.PrimaryDisplayWorkAreaRect.Top = 0;
-	OutDisplayMetrics.PrimaryDisplayWorkAreaRect.Right = 800;
-	OutDisplayMetrics.PrimaryDisplayWorkAreaRect.Bottom = 600;
-
-	OutDisplayMetrics.VirtualDisplayRect.Left = 0;
-	OutDisplayMetrics.VirtualDisplayRect.Top = 0;
-	OutDisplayMetrics.VirtualDisplayRect.Right = 800;
-	OutDisplayMetrics.VirtualDisplayRect.Bottom = 600;
+	OutDisplayMetrics.PrimaryDisplayWorkAreaRect = FHTML5Window::GetScreenRect();
+
+	OutDisplayMetrics.VirtualDisplayRect = OutDisplayMetrics.PrimaryDisplayWorkAreaRect;
+	// Total screen size of the primary monitor - or in HTML5's case- total size of the webpage. Make it arbitrary large. 
+	OutDisplayMetrics.PrimaryDisplayWidth = 4096;
+	OutDisplayMetrics.PrimaryDisplayHeight = 4096; 
+
+}
+
+TSharedRef< FGenericWindow > FHTML5Application::MakeWindow()
+{
+	return ApplicationWindow;
 }
diff --git a/Engine/Source/Runtime/Core/Private/HTML5/HTML5InputInterface.cpp b/Engine/Source/Runtime/Core/Private/HTML5/HTML5InputInterface.cpp
--- a/Engine/Source/Runtime/Core/Private/HTML5/HTML5InputInterface.cpp
+++ b/Engine/Source/Runtime/Core/Private/HTML5/HTML5InputInterface.cpp
@@ -43,7 +43,7 @@ void FHTML5InputInterface::SetMessageHandler( const TSharedRef< FGenericApplicat
 	MessageHandler = InMessageHandler;
 }
 
-void FHTML5InputInterface::Tick(float DeltaTime)
+void FHTML5InputInterface::Tick(float DeltaTime, const SDL_Event& Event)
 {
 #if PLATFORM_HTML5_WIN32 && SDL_MAJOR_VERSION < 2
 	// we have to emulate repeat
@@ -55,8 +55,6 @@ void FHTML5InputInterface::Tick(float DeltaTime)
 #define EVENT_KEY_REPEAT Event.key.repeat
 #endif
 
-	SDL_Event Event;
-	while (SDL_PollEvent(&Event)) 
 	{
 		if (Event.type == SDL_KEYDOWN)
 		{
diff --git a/Engine/Source/Runtime/Core/Private/HTML5/HTML5Window.cpp b/Engine/Source/Runtime/Core/Private/HTML5/HTML5Window.cpp
new file mode 100644
--- /dev/null
+++ b/Engine/Source/Runtime/Core/Private/HTML5/HTML5Window.cpp
@@ -0,0 +1,94 @@
+// Copyright 1998-2014 Epic Games, Inc. All Rights Reserved.
+
+#include "CorePrivate.h"
+#include "HTML5Window.h"
+
+#if !PLATFORM_HTML5_WIN32 
+#include <emscripten.h>
+#else
+#include <SDL/SDL.h>
+#endif
+
+FHTML5Window::~FHTML5Window()
+{
+	//    Use NativeWindow_Destroy() instead.
+}
+
+TSharedRef<FHTML5Window> FHTML5Window::Make()
+{
+	return MakeShareable( new FHTML5Window() );
+}
+
+FHTML5Window::FHTML5Window()
+{
+}
+
+bool FHTML5Window::GetFullScreenInfo( int32& X, int32& Y, int32& Width, int32& Height ) const
+{
+	// fix me. 
+
+	FPlatformRect ScreenRect = GetScreenRect();
+	X = ScreenRect.Left;
+	Y = ScreenRect.Top;
+	Width = ScreenRect.Right - ScreenRect.Left;
+	Height = ScreenRect.Bottom - ScreenRect.Top;
+
+	return true;
+}
+
+
+void FHTML5Window::SetOSWindowHandle(void* InWindow)
+{
+}
+
+
+FPlatformRect FHTML5Window::GetScreenRect()
+{
+	FPlatformRect ScreenRect;
+	ScreenRect.Left = 0;
+	ScreenRect.Top = 0;
+
+	int Width, Height;
+#if !PLATFORM_HTML5_WIN32
+	int fs;
+	emscripten_get_canvas_size(&Width, &Height, &fs);
+#else
+	const SDL_VideoInfo* Info = SDL_GetVideoInfo();
+	Width  = Info->current_w;
+	Height = Info->current_h;
+#endif 
+	CalculateSurfaceSize(NULL,Width,Height);
+	ScreenRect.Right = Width;
+	ScreenRect.Bottom = Height;
+	return ScreenRect;
+}
+
+void FHTML5Window::CalculateSurfaceSize(void* InWindow, int32_t& SurfaceWidth, int32_t& SurfaceHeight)
+{
+	// ensure the size is divisible by a specified amount
+	const int DividableBy = 8;
+	SurfaceWidth  = ((SurfaceWidth  + DividableBy - 1) / DividableBy) * DividableBy;
+	SurfaceHeight = ((SurfaceHeight + DividableBy - 1) / DividableBy) * DividableBy;
+
+}
+
+EWindowMode::Type FHTML5Window::GetWindowMode() const 
+{
+#if !PLATFORM_HTML5_WIN32
+	int Width,Height,FullScreen; 
+	emscripten_get_canvas_size(&Width,&Height,&FullScreen);
+	return FullScreen ? EWindowMode::Fullscreen : EWindowMode::Windowed;
+#else
+	return EWindowMode::Windowed; 
+#endif
+
+}
+
+void FHTML5Window::ReshapeWindow(int32 X, int32 Y, int32 Width, int32 Height)
+{
+	SDL_ResizeEvent Event; 
+	Event.h = Height; 
+	Event.w = Width; 
+	Event.type = SDL_VIDEORESIZE; 
+	SDL_PushEvent((SDL_Event*)&Event);
+}
diff --git a/Engine/Source/Runtime/Core/Public/HTML5/HTML5Application.h b/Engine/Source/Runtime/Core/Public/HTML5/HTML5Application.h
--- a/Engine/Source/Runtime/Core/Public/HTML5/HTML5Application.h
+++ b/Engine/Source/Runtime/Core/Public/HTML5/HTML5Application.h
@@ -3,6 +3,7 @@
 #pragma once
 
 #include "GenericApplication.h"
+#include "HTML5Window.h"
 
 /**
  * HTML5-specific application implementation.
@@ -27,6 +28,7 @@ public:
 	virtual void GetDisplayMetrics( FDisplayMetrics& OutDesktopMetrics ) const OVERRIDE;
 
 
+	TSharedRef< FGenericWindow > MakeWindow();
 private:
 
 	FHTML5Application();
@@ -35,4 +37,5 @@ private:
 private:
 
 	TSharedPtr< class FHTML5InputInterface > InputInterface;
+	TSharedRef< class FGenericWindow > ApplicationWindow;
 };
\ No newline at end of file
diff --git a/Engine/Source/Runtime/Core/Public/HTML5/HTML5InputInterface.h b/Engine/Source/Runtime/Core/Public/HTML5/HTML5InputInterface.h
--- a/Engine/Source/Runtime/Core/Public/HTML5/HTML5InputInterface.h
+++ b/Engine/Source/Runtime/Core/Public/HTML5/HTML5InputInterface.h
@@ -1,6 +1,8 @@
 // Copyright 1998-2014 Epic Games, Inc. All Rights Reserved.
 
 #pragma once
+#include <SDL/SDL.h>
+
 
 /**
  * Interface class for HTML5 input devices                 
@@ -18,8 +20,7 @@ public:
 	void SetMessageHandler( const TSharedRef< FGenericApplicationMessageHandler >& InMessageHandler );
 
 	/** Tick the interface (i.e check for new controllers) */
-	void Tick( float DeltaTime );
-
+	void Tick( float DeltaTime, const SDL_Event& Event );
 	/**
 	 * Poll for controller state and send events if needed
 	 */
diff --git a/Engine/Source/Runtime/Core/Public/HTML5/HTML5Window.h b/Engine/Source/Runtime/Core/Public/HTML5/HTML5Window.h
new file mode 100644
--- /dev/null
+++ b/Engine/Source/Runtime/Core/Public/HTML5/HTML5Window.h
@@ -0,0 +1,57 @@
+// Copyright 1998-2014 Epic Games, Inc. All Rights Reserved.
+
+#pragma once
+
+#include "GenericWindow.h"
+
+/**
+ * A platform specific implementation of FNativeWindow.
+ * Native windows provide platform-specific backing for and are always owned by an SWindow.
+ */
+ 
+class CORE_API FHTML5Window : public FGenericWindow
+{
+public:
+	~FHTML5Window();
+
+	/** Create a new FAndroidWindow.
+	 *
+	 * @param OwnerWindow		The SlateWindow for which we are crating a backing AndroidWindow
+	 * @param InParent			Parent iOS window; usually NULL.
+	 */
+	static TSharedRef<FHTML5Window> Make();
+
+	
+	virtual void* GetOSWindowHandle() const OVERRIDE { return NULL; } //can be null.
+
+	void Initialize( class FHTML5Application* const Application, const TSharedRef< FGenericWindowDefinition >& InDefinition, const TSharedPtr< FHTML5Window >& InParent, const bool bShowImmediately );
+
+	/** Returns the rectangle of the screen the window is associated with */
+	virtual bool GetFullScreenInfo( int32& X, int32& Y, int32& Width, int32& Height ) const OVERRIDE;
+
+	virtual void ReshapeWindow(int32 X, int32 Y, int32 Width, int32 Height);
+
+	virtual void SetOSWindowHandle(void*);
+
+	static FPlatformRect GetScreenRect();
+
+	static void CalculateSurfaceSize(void* InWindow, int32_t& SurfaceWidth, int32_t& SurfaceHeight);
+
+
+protected:
+	/** @return true if the native window is currently in fullscreen mode, false otherwise */
+	virtual EWindowMode::Type GetWindowMode() const OVERRIDE;
+
+private:
+	/**
+	 * Protect the constructor; only TSharedRefs of this class can be made.
+	 */
+	FHTML5Window();
+
+
+	FHTML5Application* OwningApplication;
+
+	/** Store the window region size for querying whether a point lies within the window */
+	int32 RegionX;
+	int32 RegionY;
+};
diff --git a/Engine/Source/Runtime/Launch/Private/HTML5/LaunchHTML5.cpp b/Engine/Source/Runtime/Launch/Private/HTML5/LaunchHTML5.cpp
--- a/Engine/Source/Runtime/Launch/Private/HTML5/LaunchHTML5.cpp
+++ b/Engine/Source/Runtime/Launch/Private/HTML5/LaunchHTML5.cpp
@@ -83,3 +83,5 @@ int main(int argc, char* argv[])
 
 	return 0;
 }
+
+void EmptyLinkFunctionForStaticInitializationHTML5Win32(void) {} 
diff --git a/Engine/Source/Runtime/OpenGLDrv/Private/HTML5/HTML5OpenGL.cpp b/Engine/Source/Runtime/OpenGLDrv/Private/HTML5/HTML5OpenGL.cpp
--- a/Engine/Source/Runtime/OpenGLDrv/Private/HTML5/HTML5OpenGL.cpp
+++ b/Engine/Source/Runtime/OpenGLDrv/Private/HTML5/HTML5OpenGL.cpp
@@ -21,7 +21,8 @@ void FHTML5OpenGL::ProcessExtensions( const FString& ExtensionsString )
 	bSupportsBGRA8888 = ExtensionsString.Contains(TEXT("GL_APPLE_texture_format_BGRA8888")) || ExtensionsString.Contains(TEXT("GL_IMG_texture_format_BGRA8888")) || ExtensionsString.Contains(TEXT("GL_EXT_texture_format_BGRA8888"));
 	bSupportsVertexHalfFloat = false;//ExtensionsString.Contains(TEXT("GL_OES_vertex_half_float"));
 	bSupportsTextureFloat = ExtensionsString.Contains(TEXT("GL_OES_texture_float"));
-	bSupportsTextureHalfFloat = ExtensionsString.Contains(TEXT("GL_OES_texture_half_float"));
+	bSupportsTextureHalfFloat = ExtensionsString.Contains(TEXT("GL_OES_texture_half_float")) &&
+        ExtensionsString.Contains(TEXT("GL_OES_texture_half_float_linear"));
 	bSupportsSGRB = ExtensionsString.Contains(TEXT("GL_EXT_sRGB"));
 	bSupportsColorBufferHalfFloat = ExtensionsString.Contains(TEXT("GL_EXT_color_buffer_half_float"));
 	bSupportsShaderFramebufferFetch = ExtensionsString.Contains(TEXT("GL_EXT_shader_framebuffer_fetch")) || ExtensionsString.Contains(TEXT("GL_NV_shader_framebuffer_fetch"));
@@ -48,7 +49,10 @@ void FHTML5OpenGL::ProcessExtensions( const FString& ExtensionsString )
 	// This is not color-renderable in WebGL/ANGLE (ANGLE exposes this)
 	bSupportsBGRA8888 = false;
 	// ANGLE/WEBGL_depth_texture is sort of like OES_depth_texture, you just can't upload bulk data to it (via Tex*Image2D); that should be OK?
-	bSupportsDepthTexture = ExtensionsString.Contains(TEXT("GL_ANGLE_depth_texture")) || ExtensionsString.Contains(TEXT("WEBGL_depth_texture"));
+	bSupportsDepthTexture = 
+        ExtensionsString.Contains(TEXT("WEBGL_depth_texture")) ||    // Catch "WEBGL_depth_texture", "MOZ_WEBGL_depth_texture" and "WEBKIT_WEBGL_depth_texture".
+        ExtensionsString.Contains(TEXT("GL_ANGLE_depth_texture")) || // for HTML5_WIN32 build with ANGLE
+        ExtensionsString.Contains(TEXT("GL_OES_depth_texture"));     // for a future HTML5 build without ANGLE
 
 #if !PLATFORM_HTML5_WIN32
     // The core WebGL spec has a combined GL_DEPTH_STENCIL_ATTACHMENT, unlike the core GLES2 spec.
@@ -62,13 +66,49 @@ void FHTML5OpenGL::ProcessExtensions( const FString& ExtensionsString )
 	bSupportsPackedDepthStencil = ExtensionsString.Contains(TEXT("GL_OES_packed_depth_stencil"));
 #endif
 
-	if (!ExtensionsString.Contains(TEXT("WEBGL_depth_texture")) && 		// Catch "WEBGL_depth_texture", "MOZ_WEBGL_depth_texture" and "WEBKIT_WEBGL_depth_texture".
-		!ExtensionsString.Contains(TEXT("GL_ANGLE_depth_texture")) &&	// for HTML5_WIN32 build with ANGLE
-		!ExtensionsString.Contains(TEXT("GL_OES_depth_texture")))		// for a future HTML5 build without ANGLE
-	{
+	if (!bSupportsDepthTexture) {
 		UE_LOG(LogRHI, Warning, TEXT("This browser does not support WEBGL_depth_texture. Rendering will not function since fallback code is not available."));
 	}
 
+    if (bSupportsTextureHalfFloat && !bSupportsColorBufferHalfFloat) {
+        // Initial implementations of WebGL's texture_float screwed up, and allowed
+        // rendering to fp textures, even though the underlying EXT_texture_float doesn't
+        // explicitly allow anything such.  FP rendering without explicit
+        // EXT_color_buffer_half_float may be possible, so we test for that here
+        // by checking for framebuffer completeness.  The spec is "wrong" as far as
+        // clamping and the like (which WEBGL_color_buffer_float/EXT_color_buffer_half_float
+        // fixes, but in practice it might "just work").
+        //
+        // See http://www.khronos.org/webgl/public-mailing-list/archives/1211/msg00133.html
+        // for more information.
+
+        GLenum err = glGetError();
+        if (err != GL_NO_ERROR) {
+            UE_LOG(LogRHI, Warning, TEXT("Detected OpenGL error 0x%04x before checking for implicit half-float fb support"), err);
+        }
+
+        GLuint tex, fb;
+        glGenTextures(1, &tex);
+        glBindTexture(GL_TEXTURE_2D, tex);
+        glTexImage2D(GL_TEXTURE_2D, 0, GL_RGBA, 32, 32, 0, GL_RGBA, GL_HALF_FLOAT_OES, NULL);
+        glGenFramebuffers(1, &fb);
+        glBindFramebuffer(GL_FRAMEBUFFER, fb);
+        glFramebufferTexture2D(GL_FRAMEBUFFER, GL_COLOR_ATTACHMENT0, GL_TEXTURE_2D, tex, 0);
+
+        GLenum fbstatus = glCheckFramebufferStatus(GL_FRAMEBUFFER);
+        err = glGetError();
+
+        bSupportsColorBufferHalfFloat = fbstatus == GL_FRAMEBUFFER_COMPLETE && err == GL_NO_ERROR;
+
+        if (bSupportsColorBufferHalfFloat) {
+            UE_LOG(LogRHI, Log, TEXT("Enabling implicit ColorBufferHalfFloat after checking fb completeness"));
+        }
+
+        glBindFramebuffer(GL_FRAMEBUFFER, 0);
+        glDeleteFramebuffers(1, &fb);
+        glDeleteTextures(1, &tex);
+    }
+
 	// Report shader precision
 	int Range[2];
 	glGetShaderPrecisionFormat(GL_FRAGMENT_SHADER, GL_LOW_FLOAT, Range, &ShaderLowPrecision);
@@ -125,7 +165,12 @@ void PlatformDestroyOpenGLDevice(FPlatformOpenGLDevice* Device)
 
 FPlatformOpenGLContext* PlatformCreateOpenGLContext(FPlatformOpenGLDevice* Device, void* InWindowHandle)
 {
-	Device->SharedContext->Context = SDL_SetVideoMode(800,600, 16, SDL_OPENGL);
+	int Width = 800 , Height = 600;
+#if !PLATFORM_HTML5_WIN32
+	 int fs;
+	emscripten_get_canvas_size(&Width, &Height, &fs);
+#endif
+	Device->SharedContext->Context = SDL_SetVideoMode(Width,Height, 16, SDL_OPENGL| SDL_RESIZABLE | SDL_DOUBLEBUF);
 	return Device->SharedContext; 
 }
 
@@ -172,8 +217,9 @@ void PlatformResizeGLContext( FPlatformOpenGLDevice* Device, FPlatformOpenGLCont
 
 	// we can't resize with the win32 emu right now -- we have no way of resizing the window.
 #if !PLATFORM_HTML5_WIN32
-	emscripten_set_canvas_size(SizeX, SizeY);
+	emscripten_set_canvas_size(SizeX,SizeY);
 #endif
+	Device->SharedContext->Context = SDL_SetVideoMode(SizeX,SizeY, 16, SDL_OPENGL | SDL_RESIZABLE | SDL_DOUBLEBUF );
 
 	glViewport(0, 0, SizeX, SizeY);
 	//@todo-mobile Do we need to clear here?
@@ -251,3 +297,14 @@ void PlatformReleaseRenderQuery( GLuint Query, uint64 QueryContext )
 void PlatformRestoreDesktopDisplayMode()
 {
 }
+
+extern "C"
+{
+	// callback from javascript. 
+	void execute_console_command(char* command)
+	{
+		IConsoleManager::Get().ProcessUserConsoleInput(ANSI_TO_TCHAR(command), *GWarn, NULL );
+	}
+}
+
+
diff --git a/Engine/Source/Runtime/OpenGLDrv/Private/HTML5/HTML5OpenGL.h b/Engine/Source/Runtime/OpenGLDrv/Private/HTML5/HTML5OpenGL.h
--- a/Engine/Source/Runtime/OpenGLDrv/Private/HTML5/HTML5OpenGL.h
+++ b/Engine/Source/Runtime/OpenGLDrv/Private/HTML5/HTML5OpenGL.h
@@ -103,6 +103,11 @@ struct FHTML5OpenGL : public FOpenGLES2
 
 	static FORCEINLINE bool UseES30ShadingLanguage() { return false;}
 
+	static FORCEINLINE GLsizei GetLabelObject(GLenum Type, GLuint Object, GLsizei BufferSize, ANSICHAR* OutName)
+	{
+		return 0;
+	}
+
 protected:
 	/** http://www.khronos.org/registry/webgl/extensions/WEBGL_depth_texture/ */
 	static bool bCombinedDepthStencilAttachment;
diff --git a/Engine/Source/Runtime/Renderer/Private/ForwardTranslucentRendering.cpp b/Engine/Source/Runtime/Renderer/Private/ForwardTranslucentRendering.cpp
--- a/Engine/Source/Runtime/Renderer/Private/ForwardTranslucentRendering.cpp
+++ b/Engine/Source/Runtime/Renderer/Private/ForwardTranslucentRendering.cpp
@@ -303,12 +303,15 @@ void FForwardShadingSceneRenderer::RenderTranslucency()
 		const bool bGammaSpace = !IsMobileHDR();
 		const bool bLinearHDR64 = !bGammaSpace && !IsMobileHDR32bpp();
 
+#if PLATFORM_HTML5
 		// Copy the view so emulation of framebuffer fetch works for alpha=depth.
 		// Possible optimization: this copy shouldn't be needed unless something uses fetch of depth.
 		if(bLinearHDR64 && GSupportsRenderTargetFormat_PF_FloatRGBA && (GSupportsShaderFramebufferFetch == false) && (!IsPCPlatform(GRHIShaderPlatform)))
 		{
-			//CopySceneAlpha();
+			CopySceneAlpha();
 		}
+#endif 
+
 
 		SCOPED_DRAW_EVENT(Translucency, DEC_SCENE_ITEMS);
 
diff --git a/Engine/Source/Runtime/Slate/Private/SWindow.cpp b/Engine/Source/Runtime/Slate/Private/SWindow.cpp
--- a/Engine/Source/Runtime/Slate/Private/SWindow.cpp
+++ b/Engine/Source/Runtime/Slate/Private/SWindow.cpp
@@ -298,6 +298,11 @@ void SWindow::Construct(const FArguments& InArgs)
 		WindowPosition = DisplayTopLeft + ( DisplaySize - WindowSize ) * 0.5f;
 	}
 
+#if PLATFORM_HTML5 
+	// UE expects mouse coordinates in screen space. SDL/HTML5 canvas provides in client space. 
+	// Anchor the window at the top/left corner to make sure client space coordinates and screen space coordinates match up. 
+	WindowPosition.X =  WindowPosition.Y = 0; 
+#endif 
 	this->InitialDesiredScreenPosition = WindowPosition;
 	this->InitialDesiredSize = WindowSize;
 
