#UE4 #HTML5

     Work around multiple emscripten SDL bugs. fullscreen events are now properly queued on browser fullscreen.

diff --git a/Engine/Build/HTML5/Game.html.template b/Engine/Build/HTML5/Game.html.template
--- a/Engine/Build/HTML5/Game.html.template
+++ b/Engine/Build/HTML5/Game.html.template
@@ -93,17 +93,17 @@ input[type=button]:hover {
     <script type="text/javascript">
 
 var UE4 = {
-  get console_command() {
-    var fn = Module.cwrap('execute_console_command', null, ['string']);
-    delete UE4["console_command"];
-    UE4.console_command = fn;
+  get resize_game() {
+    var fn = Module.cwrap('resize_game', null, ['number'],['number'] );
+    delete UE4["resize_game"];
+    UE4.resize_game = fn;
     return fn;
   }
 };
 
 function resizeListener(width, height) {
   console.log("resizeListener:", width, height);
-  UE4.console_command("r.SetRes " + width + "x" + height);
+  UE4.resize_game(width,height);
 }
 
 var gWasFullScreen = false;
@@ -132,8 +132,8 @@ function fullScreenChange() {
     canvas.height = height;
     canvas.setAttribute("fullscreen", "true");
 
-    console.log("r.SetRes", canvas.width, canvas.height);
-    UE4.console_command("r.SetRes " + canvas.width + "x" + canvas.height);
+    console.log("Resize: (to fullscreen)", canvas.width, canvas.height);
+    UE4.resize_game(canvas.width, canvas.height); 
   } else if (gWasFullScreen) {
     // we just left full screen
     gWasFullScreen = false;
@@ -142,8 +142,8 @@ function fullScreenChange() {
     canvas.height = gOldHeight;
     canvas.removeAttribute("fullscreen");
 
-    console.log("r.SetRes", canvas.width, canvas.height);
-    UE4.console_command("r.SetRes " + canvas.width + "x" + canvas.height);
+    console.log("Resize: (to windowed) ", canvas.width, canvas.height);
+    UE4.resize_game(canvas.width, canvas.height); 
   }
 }
 
diff --git a/Engine/Source/Programs/UnrealBuildTool/HTML5/HTML5ToolChain.cs b/Engine/Source/Programs/UnrealBuildTool/HTML5/HTML5ToolChain.cs
--- a/Engine/Source/Programs/UnrealBuildTool/HTML5/HTML5ToolChain.cs
+++ b/Engine/Source/Programs/UnrealBuildTool/HTML5/HTML5ToolChain.cs
@@ -91,7 +91,7 @@ namespace UnrealBuildTool
             // don't need UTF8 string support, and it slows string ops down
             Result += " -s UTF_STRING_SUPPORT=0";
             // export console command handler. Export main func too because default exports ( e.g Main ) are overridden if we use custom exported functions. 
-            Result += " -s EXPORTED_FUNCTIONS=\"['_main', '_execute_console_command']\" ";
+            Result += " -s EXPORTED_FUNCTIONS=\"['_main', '_resize_game']\" ";
 
             // NOTE: This may slow down the compiler's startup time!
             { 
diff --git a/Engine/Source/Runtime/Core/Private/HTML5/HTML5Application.cpp b/Engine/Source/Runtime/Core/Private/HTML5/HTML5Application.cpp
--- a/Engine/Source/Runtime/Core/Private/HTML5/HTML5Application.cpp
+++ b/Engine/Source/Runtime/Core/Private/HTML5/HTML5Application.cpp
@@ -32,7 +32,21 @@ void FHTML5Application::PollGameDeviceState( const float TimeDelta )
 		if ( Event.type == SDL_VIDEORESIZE )
 		{
 			SDL_ResizeEvent *r = (SDL_ResizeEvent*)&Event;
-			MessageHandler->OnSizeChanged(  ApplicationWindow, r->w,r->h); 
+
+			if ( r->h > 0 )
+			{
+				MessageHandler->OnSizeChanged(  ApplicationWindow, r->w,r->h); 
+			} 
+			else 
+			{
+				// Ask the game to reshape the window. 
+				int w = r->w;  
+				int h = -r->h;
+				ANSICHAR command[1048]; 
+				FCStringAnsi::Sprintf(command, "r.setRes %dx%d ", w, h);
+
+				IConsoleManager::Get().ProcessUserConsoleInput(ANSI_TO_TCHAR(command), *GWarn, NULL );
+			}
 		}
 		else 
 		{
diff --git a/Engine/Source/Runtime/Core/Private/HTML5/HTML5Window.cpp b/Engine/Source/Runtime/Core/Private/HTML5/HTML5Window.cpp
--- a/Engine/Source/Runtime/Core/Private/HTML5/HTML5Window.cpp
+++ b/Engine/Source/Runtime/Core/Private/HTML5/HTML5Window.cpp
@@ -83,9 +83,9 @@ EWindowMode::Type FHTML5Window::GetWindowMode() const
 
 void FHTML5Window::ReshapeWindow(int32 X, int32 Y, int32 Width, int32 Height)
 {
-	SDL_ResizeEvent Event; 
-	Event.h = Height; 
-	Event.w = Width; 
-	Event.type = SDL_VIDEORESIZE; 
-	SDL_PushEvent((SDL_Event*)&Event);
+	static SDL_ResizeEvent ResizeEvent; 
+	ResizeEvent.h = Height; 
+	ResizeEvent.w = Width; 
+	ResizeEvent.type = SDL_VIDEORESIZE; 
+	SDL_PushEvent((SDL_Event*)&ResizeEvent);
 }
diff --git a/Engine/Source/Runtime/OpenGLDrv/Private/HTML5/HTML5OpenGL.cpp b/Engine/Source/Runtime/OpenGLDrv/Private/HTML5/HTML5OpenGL.cpp
--- a/Engine/Source/Runtime/OpenGLDrv/Private/HTML5/HTML5OpenGL.cpp
+++ b/Engine/Source/Runtime/OpenGLDrv/Private/HTML5/HTML5OpenGL.cpp
@@ -309,9 +309,19 @@ void PlatformRestoreDesktopDisplayMode()
 extern "C"
 {
 	// callback from javascript. 
-	void execute_console_command(char* command)
+	void resize_game(int w, int h)
 	{
-		IConsoleManager::Get().ProcessUserConsoleInput(ANSI_TO_TCHAR(command), *GWarn, NULL );
+		static SDL_ResizeEvent Event; 
+		// workaround emscripten's buggy SDL implementation. 
+		Event.h = -h; 
+		Event.w = w; 
+		Event.type = SDL_VIDEORESIZE; 
+
+#if !UE_BUILD_SHIPPING 
+		emscripten_log(EM_LOG_JS_STACK | EM_LOG_WARN, "Asking UE to resize to %d x %d ", w , h );
+#endif 
+
+		SDL_PushEvent((SDL_Event*)&Event);
 	}
 }
 
