#HTML5
      TTP 335378 TappyChicken: HTML5: Need better error handling if the WebGL device context is NULL

                         Added checks to detect the presence of webgl, load game .js files if and only if we detect webgl.
                        if WebGL is not preset,  show an error message and point user to appropriate location for more diagonistics

                       ( tested w/ enabling/disabling webgl in firefox ).

diff --git a/Engine/Build/HTML5/Game.html.template b/Engine/Build/HTML5/Game.html.template
--- a/Engine/Build/HTML5/Game.html.template
+++ b/Engine/Build/HTML5/Game.html.template
@@ -172,7 +172,45 @@ function preRunHandler() {
 
 }
 
-var Module = {
+function loadJavaScriptFile(filename,isAsync){ 
+   // dynamically add script objects. 
+   var fileref=document.createElement('script');
+   fileref.setAttribute('type', 'text/javascript');
+   fileref.setAttribute('src', filename);
+   if(isAsync)
+     fileref.setAttribute('async','async');
+   document.body.appendChild(fileref);  
+}
+
+//http://www.browserleaks.com/webgl#howto-detect-webgl
+function webgl_detect()
+{
+    if (!!window.WebGLRenderingContext) {
+        var canvas = document.createElement("canvas"),
+             names = ["webgl", "experimental-webgl", "moz-webgl", "webkit-3d"],
+           context = false;
+ 
+        for(var i=0;i<4;i++) {
+            try {
+                context = canvas.getContext(names[i]);
+                if (context && typeof context.getParameter == "function") {
+                    // WebGL is enabled
+                    return true;
+                }
+            } catch(e) {}
+        }
+        // WebGL is supported, but disabled
+        return false;
+    }
+    // WebGL not supported
+    return false;
+}
+
+var Module;
+
+if (webgl_detect())
+{
+   Module = {
   preRun: [preRunHandler],
   postRun: [],
   TOTAL_MEMORY: %HEAPSIZE%,
@@ -217,18 +255,20 @@ var Module = {
   }
 };
 Module.setStatus('Downloading...');
-    </script>
-      <script type="text/javascript">
-        var tmid = Date.now();
-      </script>
-      <script src="json2.js"></script>
-      <script src="moz_binarystring.js"></script>
-      <script src="jStorage.js"></script>
-      <script src="%GAME%.data.js"></script>
-      <script async src="%CONFIG%.js"></script>
-      <script type="text/javascript">
-    var tend = Date.now();
-      console.log("Script load time: Data[js]: " + (tmid-tstart).toFixed(2) + "ms, code[js]: " + (tend-tmid).toFixed(2));
+
+  // add these scripts to the dom. 
+  loadJavaScriptFile("json2.js",false);
+  loadJavaScriptFile("moz_binarystring.js",false);
+  loadJavaScriptFile("jStorage.js",false);
+  loadJavaScriptFile("%GAME%.data.js",false);
+  loadJavaScriptFile("%CONFIG%.js",true);
+
+}
+else
+{
+  document.getElementById('display_wrapper').innerHTML ='<a href="http://get.webgl.com/">This Browser doesn\'t seem to support WebGL</a>'; 
+  document.getElementById('status').innerHTML = ""; 
+}
     </script>
   </body>
 </html>
