Add functions in html template to automatically detect browser and clamp memory requirements accordingly in case of 32 bit browsers. Add a global error handler to catch out of memory exceptions  and generate a helpful message accordingly.

diff --git a/Engine/Build/HTML5/Game.html.template b/Engine/Build/HTML5/Game.html.template
--- a/Engine/Build/HTML5/Game.html.template
+++ b/Engine/Build/HTML5/Game.html.template
@@ -206,6 +206,54 @@ function webgl_detect()
     return false;
 }
 
+
+function isBrowser64Bit()  { 
+
+   var userAgent =  window.navigator.userAgent; 
+   
+   // if we are windows and runningas as WOW64 ( windows on windows 64 ) or Win32 we are a 32 bit browser. 
+   if ( userAgent.indexOf ("Windows") > -1 && ( userAgent.indexOf("WOW64") > -1 || userAgent.indexOf("Win32") > -1 ))
+        return false; 
+
+    // all other platforms and browsers - assume 64 bit. 
+    return true; 
+
+}
+
+// generated from game.template
+// note: Packaging process looks at HTML5Engine.ini to pick up values. 
+var TOTAL_GAME_MEMORY = %HEAPSIZE%;
+
+// check max memory usage, we need to clamp it down for 32 bit browsers. 
+if (!isBrowser64Bit()) { 
+
+      var max_32bit_browser_memory =  512 * 1024 * 1024 ; // using a reasonable number, this number can change depending on the memory pressure from the underlying OS and whether or not it can give a contigous block of 512 MB memory to a 32 bit process.
+
+      if ( TOTAL_GAME_MEMORY > max_32bit_browser_memory ){ 
+         console.log (" Current Browser : " + window.navigator.userAgent );
+         console.log ( "We are running in 32 bit browser, clamping requested memory size of " + TOTAL_GAME_MEMORY + " bytes to " + max_32bit_browser_memory 
+                        + " bytes or use a 64 bit browser "
+                     );
+        TOTAL_GAME_MEMORY = max_32bit_browser_memory; 
+      }
+}
+
+// setup global error handling. 
+// make exceptions more visible. 
+window.onerror  = function(msg, url, line, column, error) {
+   if ( msg.indexOf("memory") > -1 ) { 
+
+     var message = !isBrowser64Bit() ? " We are running on a 32 bit browser, please use a 64 bit browser to avoid memory constraints "
+            : " Looks like the game needs more than the allocated " + TOTAL_GAME_MEMORY + " bytes, please edit HTML5Engine.ini and repackage "; 
+
+     alert ( "We ran out of memory : " + message );
+   }
+   else
+   {
+     alert ( "Error: " + msg );
+   }
+};
+
 var Module;
 
 if (webgl_detect())
@@ -213,7 +261,7 @@ if (webgl_detect())
    Module = {
   preRun: [preRunHandler],
   postRun: [],
-  TOTAL_MEMORY: %HEAPSIZE%,
+  TOTAL_MEMORY: TOTAL_GAME_MEMORY,
   noImageDecoding: true,
   noAudioDecoding: true,
       arguments: [%UE4CMDLINE% , filehostargument],
