Removed Quit button from html template
ttp# 348150

diff --git a/Engine/Build/HTML5/Game.html.template b/Engine/Build/HTML5/Game.html.template
--- a/Engine/Build/HTML5/Game.html.template
+++ b/Engine/Build/HTML5/Game.html.template
@@ -84,7 +84,7 @@ input[type=button]:hover {
       <input type="button" value="Fullscreen" onclick="Module.requestFullScreen(false)">
       <input type="button" value="Pause" onclick="Module['pauseMainLoop']();">
       <input type="button" value="Resume" onclick="Module['resumeMainLoop']();">
-      <input type="button" value="Quit" onclick="__exit(0)">
+      <!--input type="button" value="Quit" onclick="__exit(0)" -->
       <input type="button" value="ClearStoredGame" onclick="$.jStorage.flush()">
     </div>
 
