re-add file with correct case and relvant automation changes.

#html5

diff --git a/Engine/Build/HTML5/jstorage.js b/Engine/Build/HTML5/jstorage.js
new file mode 100644
--- /dev/null
+++ b/Engine/Build/HTML5/jstorage.js
@@ -0,0 +1,960 @@
+/*
+ * ----------------------------- JSTORAGE -------------------------------------
+ * Simple local storage wrapper to save data on the browser side, supporting
+ * all major browsers - IE6+, Firefox2+, Safari4+, Chrome4+ and Opera 10.5+
+ *
+ * Author: Andris Reinman, [REDACTED_EMAIL]
+ * Project homepage: www.jstorage.info
+ *
+ * Licensed under Unlicense:
+ *
+ * This is free and unencumbered software released into the public domain.
+ * 
+ * Anyone is free to copy, modify, publish, use, compile, sell, or
+ * distribute this software, either in source code form or as a compiled
+ * binary, for any purpose, commercial or non-commercial, and by any
+ * means.
+ * 
+ * In jurisdictions that recognize copyright laws, the author or authors
+ * of this software dedicate any and all copyright interest in the
+ * software to the public domain. We make this dedication for the benefit
+ * of the public at large and to the detriment of our heirs and
+ * successors. We intend this dedication to be an overt act of
+ * relinquishment in perpetuity of all present and future rights to this
+ * software under copyright law.
+ * 
+ * THE SOFTWARE IS PROVIDED "AS IS", WITHOUT WARRANTY OF ANY KIND,
+ * EXPRESS OR IMPLIED, INCLUDING BUT NOT LIMITED TO THE WARRANTIES OF
+ * MERCHANTABILITY, FITNESS FOR A PARTICULAR PURPOSE AND NONINFRINGEMENT.
+ * IN NO EVENT SHALL THE AUTHORS BE LIABLE FOR ANY CLAIM, DAMAGES OR
+ * OTHER LIABILITY, WHETHER IN AN ACTION OF CONTRACT, TORT OR OTHERWISE,
+ * ARISING FROM, OUT OF OR IN CONNECTION WITH THE SOFTWARE OR THE USE OR
+ * OTHER DEALINGS IN THE SOFTWARE.
+ * 
+ * For more information, please refer to <http://unlicense.org/>
+ */
+
+ (function(){
+    var
+        /* jStorage version */
+        JSTORAGE_VERSION = "0.4.8",
+
+        /* detect a dollar object or create one if not found */
+        $ = window.jQuery || window.$ || (window.$ = {}),
+
+        /* check for a JSON handling support */
+        JSON = {
+            parse:
+                window.JSON && (window.JSON.parse || window.JSON.decode) ||
+                String.prototype.evalJSON && function(str){return String(str).evalJSON();} ||
+                $.parseJSON ||
+                $.evalJSON,
+            stringify:
+                Object.toJSON ||
+                window.JSON && (window.JSON.stringify || window.JSON.encode) ||
+                $.toJSON
+        };
+
+    // Break if no JSON support was found
+    if(!("parse" in JSON) || !("stringify" in JSON)){
+        throw new Error("No JSON support found, include //cdnjs.cloudflare.com/ajax/libs/json2/20110223/json2.js to page");
+    }
+
+    var
+        /* This is the object, that holds the cached values */
+        _storage = {__jstorage_meta:{CRC32:{}}},
+
+        /* Actual browser storage (localStorage or globalStorage["domain"]) */
+        _storage_service = {jStorage:"{}"},
+
+        /* DOM element for older IE versions, holds userData behavior */
+        _storage_elm = null,
+
+        /* How much space does the storage take */
+        _storage_size = 0,
+
+        /* which backend is currently used */
+        _backend = false,
+
+        /* onchange observers */
+        _observers = {},
+
+        /* timeout to wait after onchange event */
+        _observer_timeout = false,
+
+        /* last update time */
+        _observer_update = 0,
+
+        /* pubsub observers */
+        _pubsub_observers = {},
+
+        /* skip published items older than current timestamp */
+        _pubsub_last = +new Date(),
+
+        /* Next check for TTL */
+        _ttl_timeout,
+
+        /**
+         * XML encoding and decoding as XML nodes can't be JSON'ized
+         * XML nodes are encoded and decoded if the node is the value to be saved
+         * but not if it's as a property of another object
+         * Eg. -
+         *   $.jStorage.set("key", xmlNode);        // IS OK
+         *   $.jStorage.set("key", {xml: xmlNode}); // NOT OK
+         */
+        _XMLService = {
+
+            /**
+             * Validates a XML node to be XML
+             * based on jQuery.isXML function
+             */
+            isXML: function(elm){
+                var documentElement = (elm ? elm.ownerDocument || elm : 0).documentElement;
+                return documentElement ? documentElement.nodeName !== "HTML" : false;
+            },
+
+            /**
+             * Encodes a XML node to string
+             * based on http://www.mercurytide.co.uk/news/article/issues-when-working-ajax/
+             */
+            encode: function(xmlNode) {
+                if(!this.isXML(xmlNode)){
+                    return false;
+                }
+                try{ // Mozilla, Webkit, Opera
+                    return new XMLSerializer().serializeToString(xmlNode);
+                }catch(E1) {
+                    try {  // IE
+                        return xmlNode.xml;
+                    }catch(E2){}
+                }
+                return false;
+            },
+
+            /**
+             * Decodes a XML node from string
+             * loosely based on http://outwestmedia.com/jquery-plugins/xmldom/
+             */
+            decode: function(xmlString){
+                var dom_parser = ("DOMParser" in window && (new DOMParser()).parseFromString) ||
+                        (window.ActiveXObject && function(_xmlString) {
+                    var xml_doc = new ActiveXObject("Microsoft.XMLDOM");
+                    xml_doc.async = "false";
+                    xml_doc.loadXML(_xmlString);
+                    return xml_doc;
+                }),
+                resultXML;
+                if(!dom_parser){
+                    return false;
+                }
+                resultXML = dom_parser.call("DOMParser" in window && (new DOMParser()) || window, xmlString, "text/xml");
+                return this.isXML(resultXML)?resultXML:false;
+            }
+        };
+
+
+    ////////////////////////// PRIVATE METHODS ////////////////////////
+
+    /**
+     * Initialization function. Detects if the browser supports DOM Storage
+     * or userData behavior and behaves accordingly.
+     */
+    function _init(){
+        /* Check if browser supports localStorage */
+        var localStorageReallyWorks = false;
+        if("localStorage" in window){
+            try {
+                window.localStorage.setItem("_tmptest", "tmpval");
+                localStorageReallyWorks = true;
+                window.localStorage.removeItem("_tmptest");
+            } catch(BogusQuotaExceededErrorOnIos5) {
+                // Thanks be to iOS5 Private Browsing mode which throws
+                // QUOTA_EXCEEDED_ERRROR DOM Exception 22.
+            }
+        }
+
+        if(localStorageReallyWorks){
+            try {
+                if(window.localStorage) {
+                    _storage_service = window.localStorage;
+                    _backend = "localStorage";
+                    _observer_update = _storage_service.jStorage_update;
+                }
+            } catch(E3) {/* Firefox fails when touching localStorage and cookies are disabled */}
+        }
+        /* Check if browser supports globalStorage */
+        else if("globalStorage" in window){
+            try {
+                if(window.globalStorage) {
+                    if(window.location.hostname == "localhost"){
+                        _storage_service = window.globalStorage["localhost.localdomain"];
+                    }
+                    else{
+                        _storage_service = window.globalStorage[window.location.hostname];
+                    }
+                    _backend = "globalStorage";
+                    _observer_update = _storage_service.jStorage_update;
+                }
+            } catch(E4) {/* Firefox fails when touching localStorage and cookies are disabled */}
+        }
+        /* Check if browser supports userData behavior */
+        else {
+            _storage_elm = document.createElement("link");
+            if(_storage_elm.addBehavior){
+
+                /* Use a DOM element to act as userData storage */
+                _storage_elm.style.behavior = "url(#default#userData)";
+
+                /* userData element needs to be inserted into the DOM! */
+                document.getElementsByTagName("head")[0].appendChild(_storage_elm);
+
+                try{
+                    _storage_elm.load("jStorage");
+                }catch(E){
+                    // try to reset cache
+                    _storage_elm.setAttribute("jStorage", "{}");
+                    _storage_elm.save("jStorage");
+                    _storage_elm.load("jStorage");
+                }
+
+                var data = "{}";
+                try{
+                    data = _storage_elm.getAttribute("jStorage");
+                }catch(E5){}
+
+                try{
+                    _observer_update = _storage_elm.getAttribute("jStorage_update");
+                }catch(E6){}
+
+                _storage_service.jStorage = data;
+                _backend = "userDataBehavior";
+            }else{
+                _storage_elm = null;
+                return;
+            }
+        }
+
+        // Load data from storage
+        _load_storage();
+
+        // remove dead keys
+        _handleTTL();
+
+        // start listening for changes
+        _setupObserver();
+
+        // initialize publish-subscribe service
+        _handlePubSub();
+
+        // handle cached navigation
+        if("addEventListener" in window){
+            window.addEventListener("pageshow", function(event){
+                if(event.persisted){
+                    _storageObserver();
+                }
+            }, false);
+        }
+    }
+
+    /**
+     * Reload data from storage when needed
+     */
+    function _reloadData(){
+        var data = "{}";
+
+        if(_backend == "userDataBehavior"){
+            _storage_elm.load("jStorage");
+
+            try{
+                data = _storage_elm.getAttribute("jStorage");
+            }catch(E5){}
+
+            try{
+                _observer_update = _storage_elm.getAttribute("jStorage_update");
+            }catch(E6){}
+
+            _storage_service.jStorage = data;
+        }
+
+        _load_storage();
+
+        // remove dead keys
+        _handleTTL();
+
+        _handlePubSub();
+    }
+
+    /**
+     * Sets up a storage change observer
+     */
+    function _setupObserver(){
+        if(_backend == "localStorage" || _backend == "globalStorage"){
+            if("addEventListener" in window){
+                window.addEventListener("storage", _storageObserver, false);
+            }else{
+                document.attachEvent("onstorage", _storageObserver);
+            }
+        }else if(_backend == "userDataBehavior"){
+            setInterval(_storageObserver, 1000);
+        }
+    }
+
+    /**
+     * Fired on any kind of data change, needs to check if anything has
+     * really been changed
+     */
+    function _storageObserver(){
+        var updateTime;
+        // cumulate change notifications with timeout
+        clearTimeout(_observer_timeout);
+        _observer_timeout = setTimeout(function(){
+
+            if(_backend == "localStorage" || _backend == "globalStorage"){
+                updateTime = _storage_service.jStorage_update;
+            }else if(_backend == "userDataBehavior"){
+                _storage_elm.load("jStorage");
+                try{
+                    updateTime = _storage_elm.getAttribute("jStorage_update");
+                }catch(E5){}
+            }
+
+            if(updateTime && updateTime != _observer_update){
+                _observer_update = updateTime;
+                _checkUpdatedKeys();
+            }
+
+        }, 25);
+    }
+
+    /**
+     * Reloads the data and checks if any keys are changed
+     */
+    function _checkUpdatedKeys(){
+        var oldCrc32List = JSON.parse(JSON.stringify(_storage.__jstorage_meta.CRC32)),
+            newCrc32List;
+
+        _reloadData();
+        newCrc32List = JSON.parse(JSON.stringify(_storage.__jstorage_meta.CRC32));
+
+        var key,
+            updated = [],
+            removed = [];
+
+        for(key in oldCrc32List){
+            if(oldCrc32List.hasOwnProperty(key)){
+                if(!newCrc32List[key]){
+                    removed.push(key);
+                    continue;
+                }
+                if(oldCrc32List[key] != newCrc32List[key] && String(oldCrc32List[key]).substr(0,2) == "2."){
+                    updated.push(key);
+                }
+            }
+        }
+
+        for(key in newCrc32List){
+            if(newCrc32List.hasOwnProperty(key)){
+                if(!oldCrc32List[key]){
+                    updated.push(key);
+                }
+            }
+        }
+
+        _fireObservers(updated, "updated");
+        _fireObservers(removed, "deleted");
+    }
+
+    /**
+     * Fires observers for updated keys
+     *
+     * @param {Array|String} keys Array of key names or a key
+     * @param {String} action What happened with the value (updated, deleted, flushed)
+     */
+    function _fireObservers(keys, action){
+        keys = [].concat(keys || []);
+        if(action == "flushed"){
+            keys = [];
+            for(var key in _observers){
+                if(_observers.hasOwnProperty(key)){
+                    keys.push(key);
+                }
+            }
+            action = "deleted";
+        }
+        for(var i=0, len = keys.length; i<len; i++){
+            if(_observers[keys[i]]){
+                for(var j=0, jlen = _observers[keys[i]].length; j<jlen; j++){
+                    _observers[keys[i]][j](keys[i], action);
+                }
+            }
+            if(_observers["*"]){
+                for(var j=0, jlen = _observers["*"].length; j<jlen; j++){
+                    _observers["*"][j](keys[i], action);
+                }
+            }
+        }
+    }
+
+    /**
+     * Publishes key change to listeners
+     */
+    function _publishChange(){
+        var updateTime = (+new Date()).toString();
+
+        if(_backend == "localStorage" || _backend == "globalStorage"){
+            try {
+                _storage_service.jStorage_update = updateTime;
+            } catch (E8) {
+                // safari private mode has been enabled after the jStorage initialization
+                _backend = false;
+            }
+        }else if(_backend == "userDataBehavior"){
+            _storage_elm.setAttribute("jStorage_update", updateTime);
+            _storage_elm.save("jStorage");
+        }
+
+        _storageObserver();
+    }
+
+    /**
+     * Loads the data from the storage based on the supported mechanism
+     */
+    function _load_storage(){
+        /* if jStorage string is retrieved, then decode it */
+        if(_storage_service.jStorage){
+            try{
+                _storage = JSON.parse(String(_storage_service.jStorage));
+            }catch(E6){_storage_service.jStorage = "{}";}
+        }else{
+            _storage_service.jStorage = "{}";
+        }
+        _storage_size = _storage_service.jStorage?String(_storage_service.jStorage).length:0;
+
+        if(!_storage.__jstorage_meta){
+            _storage.__jstorage_meta = {};
+        }
+        if(!_storage.__jstorage_meta.CRC32){
+            _storage.__jstorage_meta.CRC32 = {};
+        }
+    }
+
+    /**
+     * This functions provides the "save" mechanism to store the jStorage object
+     */
+    function _save(){
+        _dropOldEvents(); // remove expired events
+        try{
+            _storage_service.jStorage = JSON.stringify(_storage);
+            // If userData is used as the storage engine, additional
+            if(_storage_elm) {
+                _storage_elm.setAttribute("jStorage",_storage_service.jStorage);
+                _storage_elm.save("jStorage");
+            }
+            _storage_size = _storage_service.jStorage?String(_storage_service.jStorage).length:0;
+        }catch(E7){/* probably cache is full, nothing is saved this way*/}
+    }
+
+    /**
+     * Function checks if a key is set and is string or numberic
+     *
+     * @param {String} key Key name
+     */
+    function _checkKey(key){
+        if(typeof key != "string" && typeof key != "number"){
+            throw new TypeError("Key name must be string or numeric");
+        }
+        if(key == "__jstorage_meta"){
+            throw new TypeError("Reserved key name");
+        }
+        return true;
+    }
+
+    /**
+     * Removes expired keys
+     */
+    function _handleTTL(){
+        var curtime, i, TTL, CRC32, nextExpire = Infinity, changed = false, deleted = [];
+
+        clearTimeout(_ttl_timeout);
+
+        if(!_storage.__jstorage_meta || typeof _storage.__jstorage_meta.TTL != "object"){
+            // nothing to do here
+            return;
+        }
+
+        curtime = +new Date();
+        TTL = _storage.__jstorage_meta.TTL;
+
+        CRC32 = _storage.__jstorage_meta.CRC32;
+        for(i in TTL){
+            if(TTL.hasOwnProperty(i)){
+                if(TTL[i] <= curtime){
+                    delete TTL[i];
+                    delete CRC32[i];
+                    delete _storage[i];
+                    changed = true;
+                    deleted.push(i);
+                }else if(TTL[i] < nextExpire){
+                    nextExpire = TTL[i];
+                }
+            }
+        }
+
+        // set next check
+        if(nextExpire != Infinity){
+            _ttl_timeout = setTimeout(Math.min(_handleTTL, nextExpire - curtime, 0x7FFFFFFF));
+        }
+
+        // save changes
+        if(changed){
+            _save();
+            _publishChange();
+            _fireObservers(deleted, "deleted");
+        }
+    }
+
+    /**
+     * Checks if there's any events on hold to be fired to listeners
+     */
+    function _handlePubSub(){
+        var i, len;
+        if(!_storage.__jstorage_meta.PubSub){
+            return;
+        }
+        var pubelm,
+            _pubsubCurrent = _pubsub_last;
+
+        for(i=len=_storage.__jstorage_meta.PubSub.length-1; i>=0; i--){
+            pubelm = _storage.__jstorage_meta.PubSub[i];
+            if(pubelm[0] > _pubsub_last){
+                _pubsubCurrent = pubelm[0];
+                _fireSubscribers(pubelm[1], pubelm[2]);
+            }
+        }
+
+        _pubsub_last = _pubsubCurrent;
+    }
+
+    /**
+     * Fires all subscriber listeners for a pubsub channel
+     *
+     * @param {String} channel Channel name
+     * @param {Mixed} payload Payload data to deliver
+     */
+    function _fireSubscribers(channel, payload){
+        if(_pubsub_observers[channel]){
+            for(var i=0, len = _pubsub_observers[channel].length; i<len; i++){
+                // send immutable data that can't be modified by listeners
+                try{
+                    _pubsub_observers[channel][i](channel, JSON.parse(JSON.stringify(payload)));
+                }catch(E){};
+            }
+        }
+    }
+
+    /**
+     * Remove old events from the publish stream (at least 2sec old)
+     */
+    function _dropOldEvents(){
+        if(!_storage.__jstorage_meta.PubSub){
+            return;
+        }
+
+        var retire = +new Date() - 2000;
+
+        for(var i=0, len = _storage.__jstorage_meta.PubSub.length; i<len; i++){
+            if(_storage.__jstorage_meta.PubSub[i][0] <= retire){
+                // deleteCount is needed for IE6
+                _storage.__jstorage_meta.PubSub.splice(i, _storage.__jstorage_meta.PubSub.length - i);
+                break;
+            }
+        }
+
+        if(!_storage.__jstorage_meta.PubSub.length){
+            delete _storage.__jstorage_meta.PubSub;
+        }
+
+    }
+
+    /**
+     * Publish payload to a channel
+     *
+     * @param {String} channel Channel name
+     * @param {Mixed} payload Payload to send to the subscribers
+     */
+    function _publish(channel, payload){
+        if(!_storage.__jstorage_meta){
+            _storage.__jstorage_meta = {};
+        }
+        if(!_storage.__jstorage_meta.PubSub){
+            _storage.__jstorage_meta.PubSub = [];
+        }
+
+        _storage.__jstorage_meta.PubSub.unshift([+new Date, channel, payload]);
+
+        _save();
+        _publishChange();
+    }
+
+
+    /**
+     * JS Implementation of MurmurHash2
+     *
+     *  SOURCE: [REDACTED_PROFILE_URL]/murmurhash-js (MIT licensed)
+     *
+     * @author <a href="mailto:[REDACTED_EMAIL]">Gary Court</a>
+     * @see [REDACTED_PROFILE_URL]/murmurhash-js
+     * @author <a href="mailto:[REDACTED_EMAIL]">Austin Appleby</a>
+     * @see http://sites.google.com/site/murmurhash/
+     *
+     * @param {string} str ASCII only
+     * @param {number} seed Positive integer only
+     * @return {number} 32-bit positive integer hash
+     */
+
+    function murmurhash2_32_gc(str, seed) {
+        var
+            l = str.length,
+            h = seed ^ l,
+            i = 0,
+            k;
+
+        while (l >= 4) {
+            k =
+                ((str.charCodeAt(i) & 0xff)) |
+                ((str.charCodeAt(++i) & 0xff) << 8) |
+                ((str.charCodeAt(++i) & 0xff) << 16) |
+                ((str.charCodeAt(++i) & 0xff) << 24);
+
+            k = (((k & 0xffff) * 0x5bd1e995) + ((((k >>> 16) * 0x5bd1e995) & 0xffff) << 16));
+            k ^= k >>> 24;
+            k = (((k & 0xffff) * 0x5bd1e995) + ((((k >>> 16) * 0x5bd1e995) & 0xffff) << 16));
+
+            h = (((h & 0xffff) * 0x5bd1e995) + ((((h >>> 16) * 0x5bd1e995) & 0xffff) << 16)) ^ k;
+
+            l -= 4;
+            ++i;
+        }
+
+        switch (l) {
+            case 3: h ^= (str.charCodeAt(i + 2) & 0xff) << 16;
+            case 2: h ^= (str.charCodeAt(i + 1) & 0xff) << 8;
+            case 1: h ^= (str.charCodeAt(i) & 0xff);
+                h = (((h & 0xffff) * 0x5bd1e995) + ((((h >>> 16) * 0x5bd1e995) & 0xffff) << 16));
+        }
+
+        h ^= h >>> 13;
+        h = (((h & 0xffff) * 0x5bd1e995) + ((((h >>> 16) * 0x5bd1e995) & 0xffff) << 16));
+        h ^= h >>> 15;
+
+        return h >>> 0;
+    }
+
+    ////////////////////////// PUBLIC INTERFACE /////////////////////////
+
+    $.jStorage = {
+        /* Version number */
+        version: JSTORAGE_VERSION,
+
+        /**
+         * Sets a key's value.
+         *
+         * @param {String} key Key to set. If this value is not set or not
+         *              a string an exception is raised.
+         * @param {Mixed} value Value to set. This can be any value that is JSON
+         *              compatible (Numbers, Strings, Objects etc.).
+         * @param {Object} [options] - possible options to use
+         * @param {Number} [options.TTL] - optional TTL value
+         * @return {Mixed} the used value
+         */
+        set: function(key, value, options){
+            _checkKey(key);
+
+            options = options || {};
+
+            // undefined values are deleted automatically
+            if(typeof value == "undefined"){
+                this.deleteKey(key);
+                return value;
+            }
+
+            if(_XMLService.isXML(value)){
+                value = {_is_xml:true,xml:_XMLService.encode(value)};
+            }else if(typeof value == "function"){
+                return undefined; // functions can't be saved!
+            }else if(value && typeof value == "object"){
+                // clone the object before saving to _storage tree
+                value = JSON.parse(JSON.stringify(value));
+            }
+
+            _storage[key] = value;
+
+            _storage.__jstorage_meta.CRC32[key] = "2." + murmurhash2_32_gc(JSON.stringify(value), 0x9747b28c);
+
+            this.setTTL(key, options.TTL || 0); // also handles saving and _publishChange
+
+            _fireObservers(key, "updated");
+            return value;
+        },
+
+        /**
+         * Looks up a key in cache
+         *
+         * @param {String} key - Key to look up.
+         * @param {mixed} def - Default value to return, if key didn't exist.
+         * @return {Mixed} the key value, default value or null
+         */
+        get: function(key, def){
+            _checkKey(key);
+            if(key in _storage){
+                if(_storage[key] && typeof _storage[key] == "object" && _storage[key]._is_xml) {
+                    return _XMLService.decode(_storage[key].xml);
+                }else{
+                    return _storage[key];
+                }
+            }
+            return typeof(def) == "undefined" ? null : def;
+        },
+
+        /**
+         * Deletes a key from cache.
+         *
+         * @param {String} key - Key to delete.
+         * @return {Boolean} true if key existed or false if it didn't
+         */
+        deleteKey: function(key){
+            _checkKey(key);
+            if(key in _storage){
+                delete _storage[key];
+                // remove from TTL list
+                if(typeof _storage.__jstorage_meta.TTL == "object" &&
+                  key in _storage.__jstorage_meta.TTL){
+                    delete _storage.__jstorage_meta.TTL[key];
+                }
+
+                delete _storage.__jstorage_meta.CRC32[key];
+
+                _save();
+                _publishChange();
+                _fireObservers(key, "deleted");
+                return true;
+            }
+            return false;
+        },
+
+        /**
+         * Sets a TTL for a key, or remove it if ttl value is 0 or below
+         *
+         * @param {String} key - key to set the TTL for
+         * @param {Number} ttl - TTL timeout in milliseconds
+         * @return {Boolean} true if key existed or false if it didn't
+         */
+        setTTL: function(key, ttl){
+            var curtime = +new Date();
+            _checkKey(key);
+            ttl = Number(ttl) || 0;
+            if(key in _storage){
+
+                if(!_storage.__jstorage_meta.TTL){
+                    _storage.__jstorage_meta.TTL = {};
+                }
+
+                // Set TTL value for the key
+                if(ttl>0){
+                    _storage.__jstorage_meta.TTL[key] = curtime + ttl;
+                }else{
+                    delete _storage.__jstorage_meta.TTL[key];
+                }
+
+                _save();
+
+                _handleTTL();
+
+                _publishChange();
+                return true;
+            }
+            return false;
+        },
+
+        /**
+         * Gets remaining TTL (in milliseconds) for a key or 0 when no TTL has been set
+         *
+         * @param {String} key Key to check
+         * @return {Number} Remaining TTL in milliseconds
+         */
+        getTTL: function(key){
+            var curtime = +new Date(), ttl;
+            _checkKey(key);
+            if(key in _storage && _storage.__jstorage_meta.TTL && _storage.__jstorage_meta.TTL[key]){
+                ttl = _storage.__jstorage_meta.TTL[key] - curtime;
+                return ttl || 0;
+            }
+            return 0;
+        },
+
+        /**
+         * Deletes everything in cache.
+         *
+         * @return {Boolean} Always true
+         */
+        flush: function(){
+            _storage = {__jstorage_meta:{CRC32:{}}};
+            _save();
+            _publishChange();
+            _fireObservers(null, "flushed");
+            return true;
+        },
+
+        /**
+         * Returns a read-only copy of _storage
+         *
+         * @return {Object} Read-only copy of _storage
+        */
+        storageObj: function(){
+            function F() {}
+            F.prototype = _storage;
+            return new F();
+        },
+
+        /**
+         * Returns an index of all used keys as an array
+         * ["key1", "key2",.."keyN"]
+         *
+         * @return {Array} Used keys
+        */
+        index: function(){
+            var index = [], i;
+            for(i in _storage){
+                if(_storage.hasOwnProperty(i) && i != "__jstorage_meta"){
+                    index.push(i);
+                }
+            }
+            return index;
+        },
+
+        /**
+         * How much space in bytes does the storage take?
+         *
+         * @return {Number} Storage size in chars (not the same as in bytes,
+         *                  since some chars may take several bytes)
+         */
+        storageSize: function(){
+            return _storage_size;
+        },
+
+        /**
+         * Which backend is currently in use?
+         *
+         * @return {String} Backend name
+         */
+        currentBackend: function(){
+            return _backend;
+        },
+
+        /**
+         * Test if storage is available
+         *
+         * @return {Boolean} True if storage can be used
+         */
+        storageAvailable: function(){
+            return !!_backend;
+        },
+
+        /**
+         * Register change listeners
+         *
+         * @param {String} key Key name
+         * @param {Function} callback Function to run when the key changes
+         */
+        listenKeyChange: function(key, callback){
+            _checkKey(key);
+            if(!_observers[key]){
+                _observers[key] = [];
+            }
+            _observers[key].push(callback);
+        },
+
+        /**
+         * Remove change listeners
+         *
+         * @param {String} key Key name to unregister listeners against
+         * @param {Function} [callback] If set, unregister the callback, if not - unregister all
+         */
+        stopListening: function(key, callback){
+            _checkKey(key);
+
+            if(!_observers[key]){
+                return;
+            }
+
+            if(!callback){
+                delete _observers[key];
+                return;
+            }
+
+            for(var i = _observers[key].length - 1; i>=0; i--){
+                if(_observers[key][i] == callback){
+                    _observers[key].splice(i,1);
+                }
+            }
+        },
+
+        /**
+         * Subscribe to a Publish/Subscribe event stream
+         *
+         * @param {String} channel Channel name
+         * @param {Function} callback Function to run when the something is published to the channel
+         */
+        subscribe: function(channel, callback){
+            channel = (channel || "").toString();
+            if(!channel){
+                throw new TypeError("Channel not defined");
+            }
+            if(!_pubsub_observers[channel]){
+                _pubsub_observers[channel] = [];
+            }
+            _pubsub_observers[channel].push(callback);
+        },
+
+        /**
+         * Publish data to an event stream
+         *
+         * @param {String} channel Channel name
+         * @param {Mixed} payload Payload to deliver
+         */
+        publish: function(channel, payload){
+            channel = (channel || "").toString();
+            if(!channel){
+                throw new TypeError("Channel not defined");
+            }
+
+            _publish(channel, payload);
+        },
+
+        /**
+         * Reloads the data from browser storage
+         */
+        reInit: function(){
+            _reloadData();
+        },
+
+        /**
+         * Removes reference from global objects and saves it as jStorage
+         *
+         * @param {Boolean} option if needed to save object as simple "jStorage" in windows context
+         */
+         noConflict: function( saveInGlobal ) {
+            delete window.$.jStorage
+
+            if ( saveInGlobal ) {
+                window.jStorage = this;
+            }
+
+            return this;
+         }
+    };
+
+    // Initialize jStorage
+    _init();
+
+})();
diff --git a/Engine/Source/Programs/AutomationTool/HTML5/HTML5Platform.Automation.cs b/Engine/Source/Programs/AutomationTool/HTML5/HTML5Platform.Automation.cs
--- a/Engine/Source/Programs/AutomationTool/HTML5/HTML5Platform.Automation.cs
+++ b/Engine/Source/Programs/AutomationTool/HTML5/HTML5Platform.Automation.cs
@@ -83,8 +83,8 @@ public class HTML5Platform : Platform
             string OutDir = PackagePath;
             File.Copy(JSDir + "/json2.js", OutDir + "/json2.js", true);
             File.SetAttributes(OutDir + "/json2.js", FileAttributes.Normal);
-            File.Copy(JSDir + "/jstorage.js", OutDir + "/jstorage.js", true);
-            File.SetAttributes(OutDir + "/jstorage.js", FileAttributes.Normal);
+            File.Copy(JSDir + "/jStorage.js", OutDir + "/jStorage.js", true);
+            File.SetAttributes(OutDir + "/jStorage.js", FileAttributes.Normal);
             File.Copy(JSDir + "/moz_binarystring.js", OutDir + "/moz_binarystring.js", true);
             File.SetAttributes(OutDir + "/moz_binarystring.js", FileAttributes.Normal);
             PrintRunTime();
@@ -194,7 +194,7 @@ public class HTML5Platform : Platform
 		SC.ArchiveFiles(PackagePath, Path.GetFileName(GameExe));
 		SC.ArchiveFiles(PackagePath, Path.GetFileName(GameExe + ".mem"));
 		SC.ArchiveFiles(PackagePath, Path.GetFileName("json2.js"));
-		SC.ArchiveFiles(PackagePath, Path.GetFileName("jstorage.js"));
+		SC.ArchiveFiles(PackagePath, Path.GetFileName("jStorage.js"));
 		SC.ArchiveFiles(PackagePath, Path.GetFileName("moz_binarystring.js"));
 		SC.ArchiveFiles(PackagePath, Path.GetFileName(OutputFile));
 	}
