Fix for HTML5 failing to give useful error messages before JS  is completely parsed.

diff --git a/Engine/Build/HTML5/Game.html.template b/Engine/Build/HTML5/Game.html.template
--- a/Engine/Build/HTML5/Game.html.template
+++ b/Engine/Build/HTML5/Game.html.template
@@ -67,6 +67,9 @@ input[type=button]:hover {
   border-color: #777;
 }
 
+video {
+  display: none;
+}
     </style>
   </head>
   <body>
@@ -115,10 +118,14 @@ var UE4 = {
   }
   ,
   get on_fatal() {
-    var fn = Module.cwrap('on_fatal', null, ['string', 'string'])
-    delete UE4["on_fatal"];
-    UE4.on_fatal = fn;
-    return fn;
+    try {
+        var fn = Module.cwrap('on_fatal', null, ['string', 'string'])
+        delete UE4["on_fatal"];
+        UE4.on_fatal = fn;
+        return fn;
+    } catch(e) {
+        return function() {}
+    }
   }
 };
 
