Last resort fix for UE-9443

diff --git a/Engine/Build/HTML5/Game.html.template b/Engine/Build/HTML5/Game.html.template
--- a/Engine/Build/HTML5/Game.html.template
+++ b/Engine/Build/HTML5/Game.html.template
@@ -44,16 +44,7 @@ div#display_wrapper {
 }
 
 canvas.main {
-  border: none;
-  background: #444;
-  display: block;
-}
-
-canvas.main:not([fullscreen]) {
-  padding-right: 0;
-  margin-left: auto;
-  margin-right: auto;
-  width: 100%;
+  resize: none;
 }
 
 input[type=button] {
