Let's not nuke random SDK functions please, cause, y'know things will break.

Fixes UE-9443

diff --git a/Engine/Build/HTML5/Game.html.template b/Engine/Build/HTML5/Game.html.template
--- a/Engine/Build/HTML5/Game.html.template
+++ b/Engine/Build/HTML5/Game.html.template
@@ -44,7 +44,16 @@ div#display_wrapper {
 }
 
 canvas.main {
-  resize: none;
+  border: none;
+  background: #444;
+  display: block;
+}
+
+canvas.main:not([fullscreen]) {
+  padding-right: 0;
+  margin-left: auto;
+  margin-right: auto;
+  width: 100%;
 }
 
 input[type=button] {
@@ -171,10 +180,6 @@ function preRunHandler() {
   document.addEventListener('mozfullscreenchange', fullScreenChange, false);
   document.addEventListener('webkitfullscreenchange', fullScreenChange, false);
   document.addEventListener('MSFullscreenChange', fullScreenChange, false);
-
-  // this needs to die
-  Browser.updateCanvasDimensions = function() { };
-
 }
 
 function loadJavaScriptFile(filename,isAsync){ 
