Fix for incorrect fullscreen offset in firefox.
Fixes UE-10529

diff --git a/Engine/Build/HTML5/Game.html.template b/Engine/Build/HTML5/Game.html.template
--- a/Engine/Build/HTML5/Game.html.template
+++ b/Engine/Build/HTML5/Game.html.template
@@ -46,7 +46,6 @@ div#display_wrapper {
 canvas.main {
   border: none;
   background: #444;
-  display: block;
 }
 
 canvas.main:not([fullscreen]) {
