HTML5: An alternate template with accurate progress bar/UI styling etc.

diff --git a/Engine/Build/HTML5/GameX.html.template b/Engine/Build/HTML5/GameX.html.template
new file mode 100644
--- /dev/null
+++ b/Engine/Build/HTML5/GameX.html.template
@@ -0,0 +1,434 @@
+﻿<!DOCTYPE html>
+<html lang="en">
+<head>
+	<title>%GAME%</title>
+	<meta charset="utf-8">
+	<meta name="viewport" content="width=device-width, initial-scale=2">
+	<script src="https://code.jquery.com/jquery-2.1.3.min.js"></script>
+	<script src="http://maxcdn.bootstrapcdn.com/bootstrap/3.3.4/js/bootstrap.min.js"></script>
+	<link href="http://maxcdn.bootstrapcdn.com/bootstrap/3.3.4/css/bootstrap.min.css" rel="stylesheet">
+</head>
+<body>
+
+<style type="text/css"> 
+
+html, body, .container {
+	height: 100%;
+	overflow: hidden;
+	font-family: "Helvetica Neue", "HelveticaNeue", Helvetica, Arial, sans-serif;
+}
+
+.h4, h4{
+	margin-top: 1pt;
+	margin-bottom: 1pt;
+	font-size: 10pt; 
+}
+	
+.container {
+	display: table;
+	vertical-align: middle;
+	border: 0px;
+	border-spacing: 0px; 
+}
+
+.glyphicon-spin {
+    animation: spin 2000ms infinite linear;
+}
+
+@keyframes spin {
+    0% {
+        transform: rotate(0deg);
+    }
+    100% {
+        transform: rotate(359deg);
+    }
+}
+
+@-webkit-keyframes spin {
+    0% {
+        transform: rotate(0deg);
+    }
+    100% {
+        transform: rotate(359deg);
+    }
+}
+
+.emscripten{
+	max-height: 60%;  
+  	max-height: 60vh; 
+	min-height: 60%;  
+  	min-height: 60vh; 
+
+	align-items: center;
+	position: relative;
+	margin: 0px; 
+  	text-align: center;
+
+	display: flex;
+	justify-content: center;
+ }
+
+.texthalf {
+	height: 37%; 
+	border: 0px; 
+	padding: 0px;
+	overflow-y: scroll;
+	font-size: 2em;
+}
+
+.buttonarea { 
+	min-height: 3%;
+	border-top: 0px;
+	border-bottom: 0px;  
+	padding: 0px; 
+	margin-right: 0px;
+	margin-top: 0px;
+	margin-bottom: 0px;
+ }
+
+.btn {padding: 0px; text-align: center; min-width:150px } 
+.progress {background: rgba(245, 245, 245, 1); border: 0px solid rgba(245, 245, 245, 1); border-radius: 0px; height: 4px;}
+.progress-bar-custom {background: rgba(153, 153, 153, 1);}
+
+#progressbar { width: 50%; margin: 0 auto;}
+
+</style>
+
+<script type="text/javascript">
+
+function addLog(info, color) { 
+	$( "#logwindow" ).append( "<h4><small>" + info + " </small></h4>");
+}
+
+// combine all parallel downloads into one progress bar. 
+var totalDownloadSize = 0;
+var currentDownloadedSize = 0; 
+
+function download(urlName) { 
+
+addLog( "Starting download " + urlName );
+
+return $.ajax({
+	url: urlName,
+	dataType: "text",
+	success: function(scriptText){
+		// insert into dom - this bit is tricky, standard way makes asm.js complain about caching. 
+		// use a blob object. 
+		var scriptblob = new Blob([scriptText], { type: 'text/javascript' });
+
+		var script = document.createElement('script');
+		script.src = URL.createObjectURL(scriptblob);
+
+		script.onload = function() {
+			// we are done. 
+			URL.revokeObjectURL(script.src);
+		};
+		document.body.appendChild(script);
+		addLog ( urlName + " downloaded ");
+	},
+	xhrFields: {
+		totaladded : false,
+		lastloadedsize: 0, 
+		onprogress: function (e) {
+			if (e.lengthComputable) {
+
+				if (!this.totaladded) { 
+					totalDownloadSize += e.total; 
+					if ( e.total > 20*1024*1024){
+						addLog("possibly downloading uncompressed javascript, please host compressed javascript when in production");
+					}
+					this.totaladded = true; 
+				}
+				  
+				var deltasize = e.loaded - this.lastloadedsize; 
+				currentDownloadedSize += deltasize;
+				this.lastloadedsize = e.loaded; 
+				
+				var percentage = (currentDownloadedSize / totalDownloadSize * 100);
+
+				// update progress bar. 
+				$('.progress-bar-custom').css('width', percentage+'%').attr('aria-valuenow', percentage);   
+			}else	{
+				  console.log ("wierd, couldn't get size");
+			}
+		},
+		onload: function(e){
+		 // empty.           
+		}
+	}
+});
+};  
+
+// helper functions. 
+// http://stackoverflow.com/questions/4750015/regular-expression-to-find-urls-within-a-string
+function getHTMLGetParam(name){
+   if(name=(new RegExp('[?&]'+encodeURIComponent(name)+'=([^&]*)')).exec(location.search))
+      return decodeURIComponent(name[1]);
+}
+
+var filehostargument = "   "; 
+
+// we are serving via a server and it is unreal file server.
+if ( location.host != "" && getHTMLGetParam("cookonthefly") == "true" )
+{
+    filehostargument = "' -filehostIp=http://" + location.host + " '"; 
+}
+
+var UE4 = {
+  get resize_game() {
+    var fn = Module.cwrap('resize_game', null, ['number'],['number'] );
+    delete UE4["resize_game"];
+    UE4.resize_game = fn;
+    return fn;
+  }
+  ,
+  get on_fatal() {
+    try {
+        var fn = Module.cwrap('on_fatal', null, ['string', 'string'])
+        delete UE4["on_fatal"];
+        UE4.on_fatal = fn;
+        return fn;
+    } catch(e) {
+        return function() {}
+    }
+  },
+};
+
+//http://www.browserleaks.com/webgl#howto-detect-webgl
+function webgl_detect()
+{
+    if (!!window.WebGLRenderingContext) {
+        var canvas = document.createElement("canvas"),
+             names = ["webgl", "experimental-webgl", "moz-webgl", "webkit-3d"],
+           context = false;
+ 
+        for(var i=0;i<4;i++) {
+            try {
+                context = canvas.getContext(names[i]);
+                if (context && typeof context.getParameter == "function") {
+                    // WebGL is enabled
+                    return true;
+                }
+            } catch(e) {}
+        }
+        // WebGL is supported, but disabled
+        return false;
+    }
+    // WebGL not supported
+    return false;
+}
+
+function isBrowser64Bit()  { 
+   var userAgent =  window.navigator.userAgent; 
+    // if we are windows and runningas as WOW64 ( windows on windows 64 ) or Win32 we are a 32 bit browser. 
+    if ( userAgent.indexOf ("Windows") > -1 && ( userAgent.indexOf("WOW64") > -1 || userAgent.indexOf("Win32") > -1 ))
+        return false; 
+    // all other platforms and browsers - assume 64 bit. 
+    return true; 
+}
+
+// generated from game.template
+// note: Packaging process looks at HTML5Engine.ini to pick up values. 
+var TOTAL_GAME_MEMORY = %HEAPSIZE%;
+
+// check max memory usage, we need to clamp it down for 32 bit browsers. 
+if (!isBrowser64Bit()) { 
+		var max_32bit_browser_memory =  512 * 1024 * 1024 ; // using a reasonable number, this number can change depending on the memory pressure from the underlying OS and whether or not it can give a contiguous block of 512 MB memory to a 32 bit process.
+		if ( TOTAL_GAME_MEMORY > max_32bit_browser_memory ){ 
+			console.log (" Current Browser : " + window.navigator.userAgent );
+			console.log ( "We are running in 32 bit browser, clamping requested memory size of " + TOTAL_GAME_MEMORY + " bytes to " +max_32bit_browser_memory);
+			TOTAL_GAME_MEMORY = max_32bit_browser_memory; 
+		}
+}
+
+// setup global error handling. 
+// make exceptions more visible. 
+window.onerror  = function(msg, url, line, column, error) {
+	UE4.on_fatal(msg, error);
+  
+	// remove everything. 
+	$('#mainarea').empty(); 
+	// add alert! 
+	
+	// clean up the message.
+	msg = msg.replace(/(?:\r\n|\r|\n)/g, '<br />');
+	$ ('#mainarea').append ( '<div class="alert alert-danger text-left" style ="min-height: 10px; z-index: 2" role="alert" ><h4><small>' + msg + '</small></h4></div>');
+	
+	if ( msg.indexOf("memory") > -1 ) { 
+			  var message = !isBrowser64Bit() ? " We are running on a 32 bit browser, please use a 64 bit browser to avoid memory constraints "
+			  : " Looks like the game needs more than the allocated " + TOTAL_GAME_MEMORY + " bytes, please edit HTML5Engine.ini and repackage "; 
+			  console.log(message);
+	}
+};
+
+function preInitEmscripten(){
+	
+	// add canvas but to the dom but don't show it yet. 
+	$ ('<canvas>').
+			attr({
+				id : "canvas",
+				class: "emscripten",
+				oncontextmenu : "event.preventDefault()",
+				height: 904,
+				width: 1600, 
+			}).appendTo('#mainarea').hide(); 
+
+	// setup emscripten's canvas. 
+	Module['canvas'] = document.getElementById('canvas');  
+
+}
+
+function preRunEmscripten() { 
+}
+
+
+function postRunEmscripten() { 
+	// remove compiling message. 
+	$("#compilingmessage").remove();
+	// it still takes couple of ticks after run before anything can be drawn. show the canvas at the last possible moment. 
+	$("#canvas").show(); 
+}
+
+// create a deffered object for the data file. 
+var dataFileDeferredObject = new $.Deferred();
+
+// main emscripten module.
+var Module = {
+
+		preInit: [preInitEmscripten], 
+		preRun: [preRunEmscripten],
+		postRun: [postRunEmscripten],
+
+		TOTAL_MEMORY: TOTAL_GAME_MEMORY,
+		noImageDecoding: true,
+		noAudioDecoding: true,
+		arguments: [%UE4CMDLINE%,filehostargument],
+
+		print: (function() {
+			return addLog; 
+		})(),
+
+		printErr: function(text) {
+			console.log(text);
+		},
+
+		// state management
+		infoPrinted: false, 
+		lastcurrentDownloadedSize: 0,
+
+		// packaged data download, sets this information. 
+		setStatus: function(text) {
+		 	var matches = text.match(/([^(]+)\((\d+(\.\d+)?)\/(\d+)\)/);
+
+		 	if (matches == null){ 
+		    	addLog(text); 
+		    	return; 
+		  	}
+
+		 	if (!this.infoPrinted) { 
+			for(var download in this.dataFileDownloads){
+				addLog(" starting download " + download + " [Packaged data]");
+				// add to total size.
+				totalDownloadSize += parseInt(matches[4]); 
+			}
+		    this.infoPrinted = true; 
+		  	} 
+
+			currentDownloadedSize += (parseInt(matches[2]) - this.lastcurrentDownloadedSize); 
+			this.lastcurrentDownloadedSize = parseInt(matches[2]); 
+
+
+			var percentage = currentDownloadedSize/totalDownloadSize * 100; 
+			$('.progress-bar-custom').css('width', percentage+'%').attr('aria-valuenow', percentage);   
+
+			if (matches[2] == matches[4]){
+				// lets resolve our deffered object 
+				dataFileDeferredObject.resolve ( "Data file download finished");
+
+				for(var download in this.dataFileDownloads){
+					  addLog( download + " downloaded ")
+				}
+			}
+		},
+
+		totalDependencies: 0,
+		monitorRunDependencies: function(left) {
+			this.totalDependencies = Math.max(this.totalDependencies, left);
+		}
+    };
+
+$(document).ready(function() {
+
+	// check for webgl. 
+	if(!webgl_detect()) { 
+		$ ('#mainarea').empty();
+		$ ('#mainarea').append ( '<div class="alert alert-danger" style ="min-height: 10pt" role="alert"><a href="https://get.webgl.org/" class="alert-link">WebGL Not Detected!</a></div></div>');
+		return; 
+	}
+
+    addLog( "Starting downloads" );
+    // start downloads in parallel, print when done.  
+    var promises = [
+			'json2.js',
+			'jStorage.js',
+			'moz_binarystring.js',
+			'%GAME%.data.js',
+			'%CONFIG%.js'   
+		].map(download); 
+
+	// add the data file promise. 
+	promises.push(dataFileDeferredObject.promise()); 
+
+	$.when.apply(undefined,
+			 	promises
+		).then(	function(){ 
+				// all javascript downloads are complete.  
+				// all data file downloads are complete too. 
+
+				// at this point the only stuff is remaining it compilation and execution. 
+				// remove progress bar. 
+				$("#progressbar").remove();
+
+				// insert info about compiling; 
+
+				$ ('#mainarea').append ( '<div class="alert alert-warning" style ="min-height: 20px" role="alert" id="compilingmessage"><span class="glyphicon glyphicon-refresh glyphicon-spin"></span>  Compiling</div>');				
+				addLog ("all downloads complete ");
+			},
+		function(){
+
+				$('#progressbar').remove(); 
+				$ ('#mainarea').append ( '<div class="alert alert-danger" style ="min-height: 20px" role="alert">One or more downloads failed! Please refresh this page</div>');
+
+				addLog("Download failed");
+			}
+		);
+ 
+});
+
+</script>
+
+<div class="emscripten" id= "mainarea">
+      <div class="progress" id = "progressbar">
+        <div  class="progress-bar progress-bar-custom" 
+              role="progressbar" 
+              aria-valuenow="0" 
+              aria-valuemin="0"
+              aria-valuemax="100"
+              style="width: 0%; "
+              >
+        </div>
+      </div>
+ </div>
+</div>
+	<div class="row buttonarea text-center" id="buttonrow">
+		<div class="col-sm-2 text-center"></div>
+		<div class="col-sm-2 text-center"><button type="button" class="btn btn-primary" onclick="Module['pauseMainLoop']();"  >Pause</button></div>
+		<div class="col-sm-2 text-center"><button type="button" class="btn btn-primary" onclick="Module['resumeMainLoop']();" >Resume</button></div>
+		<div class="col-sm-2 text-center"><button type="button" class="btn btn-primary" onclick="$.jStorage.flush()">Clear Stored Game</button></div>
+		<div class="col-sm-2 text-center"><button type="button" class="btn btn-primary"  id="fullscreen_request" >FullScreen</button></div>
+		<div class="col-sm-2 text-center"></div>
+ 	</div>
+<div class="texthalf text-normal jumbotron" id="logwindow"></div>
+</div>
+</body>
+</html>
\ No newline at end of file
diff --git a/Engine/Source/Developer/HTML5/HTML5PlatformEditor/Classes/HTML5TargetSettings.h b/Engine/Source/Developer/HTML5/HTML5PlatformEditor/Classes/HTML5TargetSettings.h
--- a/Engine/Source/Developer/HTML5/HTML5PlatformEditor/Classes/HTML5TargetSettings.h
+++ b/Engine/Source/Developer/HTML5/HTML5PlatformEditor/Classes/HTML5TargetSettings.h
@@ -50,16 +50,24 @@ public:
  	 */
  	UPROPERTY(GlobalConfig, EditAnywhere, Category=Memory, Meta = (DisplayName = "Port to use when deploying game from the editor", ClampMin="49152", ClampMax="65535"))
  	int32 DeployServerPort;
+
+
+	/**
+	* Setting to control which template to use
+	*/
+	UPROPERTY(GlobalConfig, EditAnywhere, Category = Packaging, Meta = (DisplayName = "Use new template [experimental]"))
+	bool UseExperimentalTemplate;
+
 	/**
 	 * Use a loading level and download maps during transitions.                                                                     
 	 */
-	UPROPERTY(GlobalConfig, EditAnywhere, Category = Packaging, Meta = (DisplayName = "Download Maps on the fly [experimental]"))
+	UPROPERTY(GlobalConfig, EditAnywhere, Category = Packaging, Meta = (DisplayName = "Download maps on the fly [experimental]"))
 	bool UseAsyncLevelLoading;
 
 	/**
 	 * Generate Delta Pak files for these level transitions.
 	 */
-	UPROPERTY(GlobalConfig, EditAnywhere, Category = Packaging, Meta = (DisplayName = "Level transitions for delta paks [experimental]"))
+	UPROPERTY(GlobalConfig, EditAnywhere, Category = Packaging, Meta = (DisplayName = "Level transitions for delta paks [experimental,depends on download maps]"))
 	TArray<FHTML5LevelTransitions> LevelTransitions;
 
 };
diff --git a/Engine/Source/Programs/AutomationTool/HTML5/HTML5Platform.Automation.cs b/Engine/Source/Programs/AutomationTool/HTML5/HTML5Platform.Automation.cs
--- a/Engine/Source/Programs/AutomationTool/HTML5/HTML5Platform.Automation.cs
+++ b/Engine/Source/Programs/AutomationTool/HTML5/HTML5Platform.Automation.cs
@@ -99,7 +99,12 @@ public class HTML5Platform : Platform
         File.SetAttributes(Path.Combine(PackagePath, GameExe) + ".mem", FileAttributes.Normal);
 
         // put the HTML file to the package directory
-        string TemplateFile = Path.Combine(CombinePaths(CmdEnv.LocalRoot, "Engine"), "Build", "HTML5", "Game.html.template");
+
+        bool UseExperimentalTemplate = false;
+        ConfigCache.GetBool("/Script/HTML5PlatformEditor.HTML5TargetSettings", "UseExperimentalTemplate", out UseExperimentalTemplate);
+        string TemplateFileName = UseExperimentalTemplate ? "GameX.html.template" : "Game.html.template";
+
+        string TemplateFile = Path.Combine(CombinePaths(CmdEnv.LocalRoot, "Engine"), "Build", "HTML5", TemplateFileName);
         string OutputFile = Path.Combine(PackagePath, (Params.ClientConfigsToBuild[0].ToString() != "Development" ? (Params.ShortProjectName + "-HTML5-" + Params.ClientConfigsToBuild[0].ToString()) : Params.ShortProjectName)) + ".html";
 
         // find Heap Size.
