UE-8662: unresponsive script warning when using development buids.

  - don't include debuggin info / minify development builds - cuts development binary size from 200Mb to 69Mb.
  - generate symbols file
  - modify template to download this symbols file on the fly on a crash and use it generate an unmangled and unminified callstack. ( we now have real callstacks for shippign builds!)

diff --git a/Engine/Build/HTML5/GameX.html.template b/Engine/Build/HTML5/GameX.html.template
--- a/Engine/Build/HTML5/GameX.html.template
+++ b/Engine/Build/HTML5/GameX.html.template
@@ -244,21 +244,53 @@ if (!isBrowser64Bit()) {
 // setup global error handling. 
 // make exceptions more visible. 
 window.onerror  = function(msg, url, line, column, error) {
-	UE4.on_fatal(msg, error);
-  
 	// remove everything. 
 	$('#mainarea').empty(); 
 	// add alert! 
 	
 	// clean up the message.
-	msg = msg.replace(/(?:\r\n|\r|\n)/g, '<br />');
-	$ ('#mainarea').append ( '<div class="alert alert-danger text-left" style ="min-height: 10px; z-index: 2" role="alert" ><h4><small>' + msg + '</small></h4></div>');
+	var htmlmsg = msg.replace(/(?:\r\n|\r|\n)/g, '<br />');
+	$ ('#mainarea').append ( '<div class="alert alert-danger text-left" style ="min-height: 10px; z-index: 2" role="alert" ><h4><small>' + htmlmsg + '</small></h4></div>');
 	
 	if ( msg.indexOf("memory") > -1 ) { 
 			  var message = !isBrowser64Bit() ? " We are running on a 32 bit browser, please use a 64 bit browser to avoid memory constraints "
 			  : " Looks like the game needs more than the allocated " + TOTAL_GAME_MEMORY + " bytes, please edit HTML5Engine.ini and repackage "; 
 			  console.log(message);
 	}
+    else
+    {
+        addLog( " downloading symbols file ");
+        // get non minified symbols and demangle the callstack. 
+		$.ajax({
+			  url : "%CONFIG%.js.symbols",
+			  dataType: "text",
+			  success : function (data) {	
+                  addLog(" symbols file downloaded, demangling..");
+				  var mangledcallstack = msg; 
+				  var maplines = data.match(/[^\r\n]+/g);
+				  var msglines = mangledcallstack.match(/[^\r\n]+/g);
+				  var sourcemap = {};
+				  for (var i = 0; i < maplines.length; i++) {
+					  var tokens =  maplines[i].split(":"); 
+					  sourcemap[tokens[0]] = tokens[1];
+				  }
+				  var callstack = "";
+				  
+				  for (var i = 0; i < msglines.length; i++) {
+					  var tokens =  msglines[i].split("@"); 
+					  if(typeof sourcemap[tokens[0]] !== 'undefined') {
+						  callstack  += demangle(sourcemap[tokens[0]]) + "\n";
+					  }
+				  }
+				  
+				  $('#mainarea').empty(); 
+				  var htmlmsg = callstack.replace(/(?:\r\n|\r|\n)/g, '<br />');
+				  $ ('#mainarea').append ( '<div class="alert alert-danger text-left" style ="min-height: 10px; z-index: 2" role="alert" ><h4><small>' + htmlmsg + '</small></h4></div>');
+				  
+				  addLog( " callstack ready ..");
+			  }
+		  });
+    }
 };
 
 function preInitEmscripten(){
diff --git a/Engine/Source/Programs/AutomationTool/HTML5/HTML5Platform.Automation.cs b/Engine/Source/Programs/AutomationTool/HTML5/HTML5Platform.Automation.cs
--- a/Engine/Source/Programs/AutomationTool/HTML5/HTML5Platform.Automation.cs
+++ b/Engine/Source/Programs/AutomationTool/HTML5/HTML5Platform.Automation.cs
@@ -94,6 +94,7 @@ public class HTML5Platform : Platform
         {
             File.Copy(Path.Combine(Path.GetDirectoryName(Params.ProjectGameExeFilename), GameExe), Path.Combine(PackagePath, GameExe), true);
             File.Copy(Path.Combine(Path.GetDirectoryName(Params.ProjectGameExeFilename), GameExe) + ".mem", Path.Combine(PackagePath, GameExe) + ".mem", true);
+			File.Copy(Path.Combine(Path.GetDirectoryName(Params.ProjectGameExeFilename), GameExe) + ".symbols", Path.Combine(PackagePath, GameExe) + ".symbols", true);
         }
         File.SetAttributes(Path.Combine(PackagePath, GameExe), FileAttributes.Normal);
         File.SetAttributes(Path.Combine(PackagePath, GameExe) + ".mem", FileAttributes.Normal);
@@ -171,9 +172,10 @@ public class HTML5Platform : Platform
 
 				if (LineStr.Contains("%CONFIG%"))
 				{
+					string TempGameName = InGameName;
 					if (IsContentOnly)
-						InGameName = "UE4Game";
-					LineStr = LineStr.Replace("%CONFIG%", (InGameConfiguration != "Development" ? (InGameName + "-HTML5-" + InGameConfiguration) : InGameName));
+						TempGameName = "UE4Game";
+					LineStr = LineStr.Replace("%CONFIG%", (InGameConfiguration != "Development" ? (TempGameName + "-HTML5-" + InGameConfiguration) : TempGameName));
 				}
 
 				if (LineStr.Contains("%UE4CMDLINE%"))
@@ -246,6 +248,7 @@ public class HTML5Platform : Platform
 		SC.ArchiveFiles(PackagePath, Path.GetFileName(FinalDataLocation + ".js"));
 		SC.ArchiveFiles(PackagePath, Path.GetFileName(GameExe));
 		SC.ArchiveFiles(PackagePath, Path.GetFileName(GameExe + ".mem"));
+		SC.ArchiveFiles(PackagePath, Path.GetFileName(GameExe + ".symbols"));
 		SC.ArchiveFiles(PackagePath, Path.GetFileName("json2.js"));
 		SC.ArchiveFiles(PackagePath, Path.GetFileName("jStorage.js"));
 		SC.ArchiveFiles(PackagePath, Path.GetFileName("moz_binarystring.js"));
diff --git a/Engine/Source/Programs/UnrealBuildTool/HTML5/HTML5ToolChain.cs b/Engine/Source/Programs/UnrealBuildTool/HTML5/HTML5ToolChain.cs
--- a/Engine/Source/Programs/UnrealBuildTool/HTML5/HTML5ToolChain.cs
+++ b/Engine/Source/Programs/UnrealBuildTool/HTML5/HTML5ToolChain.cs
@@ -198,7 +198,7 @@ namespace UnrealBuildTool
 				}
 				if (LinkEnvironment.Config.Target.Configuration == CPPTargetConfiguration.Development)
 				{
-					Result += " -O2 -s ASM_JS=1 -s OUTLINING_LIMIT=110000   -g2 ";
+					Result += " -O2 -s ASM_JS=1 -s OUTLINING_LIMIT=110000 ";
 				}
 				if (LinkEnvironment.Config.Target.Configuration == CPPTargetConfiguration.Shipping)
 				{
@@ -516,7 +516,7 @@ namespace UnrealBuildTool
 
 		    FileItem OutputBC = FileItem.GetItemByPath(LinkEnvironment.Config.OutputFilePath.Replace(".js", ".bc").Replace(".html", ".bc"));
 		    LinkAction.ProducedItems.Add(OutputBC);
-		    LinkAction.CommandArguments += string.Format(" --save-bc \"{0}\"", OutputBC.AbsolutePath);
+			LinkAction.CommandArguments += " --emit-symbol-map " + string.Format(" --save-bc \"{0}\"", OutputBC.AbsolutePath);
 
      		LinkAction.StatusDescription = Path.GetFileName(OutputFile.AbsolutePath);
 			LinkAction.OutputEventHandler = new DataReceivedEventHandler(RemoteOutputReceivedEventHandler);
diff --git a/Engine/Source/Runtime/Core/Public/HTML5/HTML5AssertionMacros.h b/Engine/Source/Runtime/Core/Public/HTML5/HTML5AssertionMacros.h
--- a/Engine/Source/Runtime/Core/Public/HTML5/HTML5AssertionMacros.h
+++ b/Engine/Source/Runtime/Core/Public/HTML5/HTML5AssertionMacros.h
@@ -34,6 +34,8 @@ inline void html5_break_msg(const char* msg, const char* file, int line) {
 		var InMsg = Pointer_stringify($0);
 		var InFile = Pointer_stringify($1);
 		alert('Expression ('+InMsg+') failed in '+InFile+':'+$2+'!\nCheck console for details.\n'); 
+		var callstack = new Error; 
+		throw callstack.stack; 
 	}, msg, file, line);
 }
 
