Updated HTML5 templates to fix fullscreen issues across browsers.

diff --git a/Engine/Build/HTML5/Game.html.template b/Engine/Build/HTML5/Game.html.template
--- a/Engine/Build/HTML5/Game.html.template
+++ b/Engine/Build/HTML5/Game.html.template
@@ -137,48 +137,9 @@ var gWasFullScreen = false;
 var gOldWidth, gOldHeight;
 
 function fullScreenChange() {
-  var fsElem = document['fullScreenElement'] ||
-               document['mozFullScreenElement'] ||
-               document['webkitFullScreenElement'] ||
-               document['msFullScreenElement'] ||
-               document['webkitCurrentFullScreenElement'];
-  var canvas = document.getElementById("canvas");
-
-  if (fsElem == canvas || fsElem == canvas.parentNode) {
-    // we just went full screen
-    gWasFullScreen = true;
-
-    // update the canvas size
-    gOldWidth = canvas.width;
-    gOldHeight = canvas.height;
-
-    var width = window.screen.width;
-    var height = window.screen.height;
-
-    canvas.width = width;
-    canvas.height = height;
-    canvas.setAttribute("fullscreen", "true");
-
-    console.log("Resize: (to fullscreen)", canvas.width, canvas.height);
-    UE4.resize_game(canvas.width, canvas.height); 
-  } else if (gWasFullScreen) {
-    // we just left full screen
-    gWasFullScreen = false;
-
-    canvas.width = gOldWidth;
-    canvas.height = gOldHeight;
-    canvas.removeAttribute("fullscreen");
-
-    console.log("Resize: (to windowed) ", canvas.width, canvas.height);
-    UE4.resize_game(canvas.width, canvas.height); 
-  }
 }
 
 function preRunHandler() {
-  document.addEventListener('fullscreenchange', fullScreenChange, false);
-  document.addEventListener('mozfullscreenchange', fullScreenChange, false);
-  document.addEventListener('webkitfullscreenchange', fullScreenChange, false);
-  document.addEventListener('MSFullscreenChange', fullScreenChange, false);
 }
 
 function loadJavaScriptFile(filename,isAsync){ 
diff --git a/Engine/Build/HTML5/GameX.html.template b/Engine/Build/HTML5/GameX.html.template
--- a/Engine/Build/HTML5/GameX.html.template
+++ b/Engine/Build/HTML5/GameX.html.template
@@ -14,7 +14,6 @@
 
 html, body, .container {
 	height: 100%;
-	overflow: hidden;
 	font-family: "Helvetica Neue", "HelveticaNeue", Helvetica, Arial, sans-serif;
 }
 
@@ -53,20 +52,45 @@ html, body, .container {
     }
 }
 
-.emscripten{
-	max-height: 60%;  
-  	max-height: 60vh; 
-	min-height: 60%;  
-  	min-height: 60vh; 
-
+.wrapper{
+	position: relative;
+	margin: 1em auto 10px auto;
+	text-align: center;
+	min-width: 640px;
+	min-height: 480px; 
+	max-width: 95%;
+	display: block; 
 	align-items: center;
 	position: relative;
-	margin: 0px; 
-  	text-align: center;
-
-	display: flex;
+	text-align: center;
 	justify-content: center;
- }
+}
+ 
+ .emscripten {
+ 
+	padding-right: 0;
+	margin-left: auto;
+	margin-right: auto;
+	display: block;
+	display: -webkit-box;
+	display: -moz-box;
+	display: box;
+
+	-webkit-box-align: center;
+	-moz-box-align: center;
+	box-align: center;
+
+	-webkit-box-pack: center;
+	-moz-box-pack: center;
+	box-pack: center;
+}
+
+#canvas:not([fullscreen]) {
+  padding-right: 0;
+  margin-left: auto;
+  margin-right: auto;
+  width: 100%;
+}
 
 .texthalf {
 	height: 37%; 
@@ -89,7 +113,12 @@ html, body, .container {
 .btn {padding: 0px; text-align: center; min-width:150px } 
 .progress {background: rgba(245, 245, 245, 1); border: 0px solid rgba(245, 245, 245, 1); border-radius: 0px; height: 4px;}
 .progress-bar-custom {background: rgba(153, 153, 153, 1);}
-
+.centered-axis-xy {
+    position: absolute;
+    left: 50%;
+    top: 50%;
+    transform: translate(-50%,-50%);
+}
 #progressbar { width: 50%; margin: 0 auto;}
 
 </style>
@@ -250,7 +279,7 @@ window.onerror  = function(msg, url, line, column, error) {
 	
 	// clean up the message.
 	var htmlmsg = msg.replace(/(?:\r\n|\r|\n)/g, '<br />');
-	$ ('#mainarea').append ( '<div class="alert alert-danger text-left" style ="min-height: 10px; z-index: 2" role="alert" ><h4><small>' + htmlmsg + '</small></h4></div>');
+	$ ('#mainarea').append ( '<div class="alert alert-danger text-left centered-axis-xy" style ="min-height: 10px; z-index: 2" role="alert" ><h4><small>' + htmlmsg + '</small></h4></div>');
 	
 	if ( msg.indexOf("memory") > -1 ) { 
 			  var message = !isBrowser64Bit() ? " We are running on a 32 bit browser, please use a 64 bit browser to avoid memory constraints "
@@ -282,10 +311,9 @@ window.onerror  = function(msg, url, line, column, error) {
 						  callstack  += demangle(sourcemap[tokens[0]]) + "\n";
 					  }
 				  }
-				  
-				  $('#mainarea').empty(); 
+				  $('#mainarea').empty();
 				  var htmlmsg = callstack.replace(/(?:\r\n|\r|\n)/g, '<br />');
-				  $ ('#mainarea').append ( '<div class="alert alert-danger text-left" style ="min-height: 10px; z-index: 2" role="alert" ><h4><small>' + htmlmsg + '</small></h4></div>');
+				  $ ('#mainarea').append ( '<div class="alert alert-danger text-left centered-axis-xy" style ="min-height: 10px; z-index: 2" role="alert" ><h4><small>' + htmlmsg + '</small></h4></div>');
 				  
 				  addLog( " callstack ready ..");
 			  }
@@ -394,7 +422,7 @@ $(document).ready(function() {
 	// check for webgl. 
 	if(!webgl_detect()) { 
 		$ ('#mainarea').empty();
-		$ ('#mainarea').append ( '<div class="alert alert-danger" style ="min-height: 10pt" role="alert"><a href="https://get.webgl.org/" class="alert-link">WebGL Not Detected!</a></div></div>');
+		$ ('#mainarea').append ( '<div class="alert alert-danger centered-axis-xy" style ="min-height: 10pt" role="alert"><a href="https://get.webgl.org/" class="alert-link">WebGL Not Detected!</a></div></div>');
 		return; 
 	}
 
@@ -423,24 +451,27 @@ $(document).ready(function() {
 
 				// insert info about compiling; 
 
-				$ ('#mainarea').append ( '<div class="alert alert-warning" style ="min-height: 20px" role="alert" id="compilingmessage"><span class="glyphicon glyphicon-refresh glyphicon-spin"></span>  Compiling</div>');				
+				$ ('#mainarea').append ( '<div class="alert alert-warning centered-axis-xy" style ="min-height: 20px" role="alert" id="compilingmessage"><span class="glyphicon glyphicon-refresh glyphicon-spin"></span>  Compiling</div>');				
 				addLog ("all downloads complete ");
 			},
 		function(){
 
 				$('#progressbar').remove(); 
-				$ ('#mainarea').append ( '<div class="alert alert-danger" style ="min-height: 20px" role="alert">One or more downloads failed! Please refresh this page</div>');
+				$ ('#mainarea').append ( '<div class="alert alert-danger centered-axis-xy" style ="min-height: 20px" role="alert">One or more downloads failed! Please refresh this page</div>');
 
 				addLog("Download failed");
 			}
 		);
+		
+		$(document).on('webkitfullscreenchange mozfullscreenchange fullscreenchange MSFullscreenChange',  function() { 
+		});
  
 });
 
 </script>
 
-<div class="emscripten" id= "mainarea">
-      <div class="progress" id = "progressbar">
+<div class="wrapper" id= "mainarea">
+      <div class="progress centered-axis-xy" id = "progressbar">
         <div  class="progress-bar progress-bar-custom" 
               role="progressbar" 
               aria-valuenow="0" 
@@ -460,7 +491,7 @@ $(document).ready(function() {
 		<div class="col-sm-2 text-center"><button type="button" class="btn btn-primary"  id="fullscreen_request" >FullScreen</button></div>
 		<div class="col-sm-2 text-center"></div>
  	</div>
-<div class="texthalf text-normal jumbotron" id="logwindow"></div>
+<div class="texthalf text-normal jumbotron " id="logwindow"></div>
 </div>
 </body>
-</html>
\ No newline at end of file
+</html>
