#html5 fix template to properly pick compressed mem/data file.

diff --git a/Engine/Build/HTML5/Game.html.template b/Engine/Build/HTML5/Game.html.template
--- a/Engine/Build/HTML5/Game.html.template
+++ b/Engine/Build/HTML5/Game.html.template
@@ -249,6 +249,10 @@ if (webgl_detect())
   printErr: function(text) {
     console.log(text);
   },
+  // file location. 
+  locateFile : function (name){ 
+    return name + ".gz";
+  },
   canvas: document.getElementById('canvas'),
   setStatus: function(text) {
     if (Module.setStatus.interval) clearInterval(Module.setStatus.interval);
