don't mix https with http in the template  #html5

diff --git a/Engine/Build/HTML5/GameX.html.template b/Engine/Build/HTML5/GameX.html.template
--- a/Engine/Build/HTML5/GameX.html.template
+++ b/Engine/Build/HTML5/GameX.html.template
@@ -4,7 +4,7 @@
 	<title>%GAME%</title>
 	<meta charset="utf-8">
 	<meta name="viewport" content="width=device-width, initial-scale=2">
-	<script src="https://code.jquery.com/jquery-2.1.3.min.js"></script>
+	<script src="http://code.jquery.com/jquery-2.1.3.min.js"></script>
 	<script src="http://maxcdn.bootstrapcdn.com/bootstrap/3.3.4/js/bootstrap.min.js"></script>
 	<link href="http://maxcdn.bootstrapcdn.com/bootstrap/3.3.4/css/bootstrap.min.css" rel="stylesheet">
 </head>
