HTML5 - split single large template file to smaller component (e.g. js & css) files

custom template files are also supported on a per-project basis:

copy
.../Engine/Build/HTML5/GameX.*.template
to
<project>/Build/HTML5/.

build will automatically pick it the project's path (otherwise fallback to the Engine's version)

this is based on GitHub PR #4780

#rn

diff --git a/Engine/Build/HTML5/GameX.html.template b/Engine/Build/HTML5/GameX.html.template
--- a/Engine/Build/HTML5/GameX.html.template
+++ b/Engine/Build/HTML5/GameX.html.template
@@ -1,246 +1,29 @@
 ﻿<!DOCTYPE html>
+
+<!--
+     HTML code used with Epic Games HTML5 projects
+
+     much of this is for UE4 development purposes.
+
+     to create a custom HTML file for your project:
+     - make a copy of this file - or make one from scratch
+     - and put it in: "your project folder"/Build/HTML/GameX.html.template
+-->
+
 <html lang="en">
 <head>
 	<title>%GAME%</title>
 	<meta charset="utf-8">
 	<meta name="viewport" content="width=device-width, initial-scale=1">
 
-	<script>
-		// THIS IS TEMP -- to be removed when the following has been removed as well
-		// .../Engine/Source/Runtime/Engine/Private/ActiveSound.cpp
-		//     // for velocity-based effects like doppler
-		//     ParseParams.Velocity = (ParseParams.Transform.GetTranslation() - LastLocation) / DeltaTime;
-		window.AudioContext = ( window.AudioContext || window.webkitAudioContext || null );
-		if ( AudioContext ) {
-			var ue4_hacks = {}; // making this obvious...
-			ue4_hacks.ctx = new AudioContext();
-			ue4_hacks.panner = ue4_hacks.ctx.createPanner();
-			ue4_hacks.panner.__proto__.setVelocity = ( ue4_hacks.panner.__proto__.setVelocity || function(){} );
-		}
-	</script>
-
-	<script>
-		// Minimum WebGL version that the page needs in order to run. UE4 will attempt to use WebGL 2 if available.
-		// Set this to 1 to run with a WebGL 1 fallback if the graphical features required by your UE4 project are
-		// low enough that they do not strictly require WebGL 2.
-		const requiredWebGLVersion = 1;
-
-		// Add ?webgl1 GET param to explicitly test the WebGL 1 fallback version even if browser does support WebGL 2.
-		const explicitlyUseWebGL1 = (location.search.indexOf('webgl1') != -1);
-
-		// Deduce which version to load up.
-		var supportsWasm = (typeof WebAssembly === 'object' && typeof WebAssembly.Memory === 'function');
-
-		// By default, Shipping builds serve out compressed assets, and Development builds serve uncompressed assets.
-		// Change the following line to customize. When hosting UE4 builds live on a production CDN, compression
-		// should always be enabled, since uncompressed files are too huge to be downloaded over the web.
-		const serveCompressedAssets = %SERVE_COMPRESSED%;
-
-		// For the large .data file, there's two ways to manage compression: either UE4 UnrealPak tool can compress it in engine, or
-		// it can be gzip compressed on disk like other assets. Compressing via UnrealPak has the advantage of storing a smaller data
-		// file to IndexedDB, whereas gzip compressing to disk has the advantage of starting up the page slightly faster.
-		// If true, serve out 'UE4Game.data.gz', if false, serve out 'UE4Game.data'.
-//		const dataFileIsGzipCompressed = false;
-
-		var Module = {
-			// state management
-			infoPrinted: false,
-			lastcurrentDownloadedSize: 0,
-			totalDependencies: 0,
-			dataBytesStoredInIndexedDB: 0, // Track how much data is currently stored in IndexedDB.
-
-			assetDownloadProgress: {}, // Track how many bytes of each needed asset has been downloaded so far.
-
-			UE4_indexedDBName: 'UE4_assetDatabase_%GAME%', // TODO Nick: this should be an ascii ID string without special characters that is unique to the project that is being packaged
-			UE4_indexedDBVersion: %TIMESTAMP%, // Bump this number to invalidate existing IDB storages in browsers.
-		};
-
-		// Tests if type === 'browser' or type === 'os' is 64-bit or not.
-		function heuristicIs64Bit(type) {
-			function contains(str, substrList) { for(var i in substrList) if (str.indexOf(substrList[i]) != -1) return true; return false; }
-			var ua = (navigator.userAgent + ' ' + navigator.oscpu + ' ' + navigator.platform).toLowerCase();
-			if (contains(ua, ['wow64'])) return type === 'os'; // 32bit browser on 64bit OS
-			if (contains(ua, ['x86_64', 'amd64', 'ia64', 'win64', 'x64', 'arm64', 'irix64', 'mips64', 'ppc64', 'sparc64'])) return true;
-			if (contains(ua, ['i386', 'i486', 'i586', 'i686', 'x86', 'arm7', 'android', 'mobile', 'win32'])) return false;
-			if (contains(ua, ['intel mac os'])) return true;
-			return false;
-		}
-
-		// For best stability on 32-bit browsers, allocate asm.js/WebAssembly heap up front before proceeding
-		// to load any other page content. This mitigates the chances that loading up page assets first would
-		// fragment the memory area of the browser process.
-		var pageSize = 64 * 1024;
-		var heuristic64BitBrowser = heuristicIs64Bit('browser');
-		function alignPageUp(size) { return pageSize * Math.ceil(size / pageSize); }
-
-		// The application should not be able to allocate more than MAX bytes of memory. If the application
-		// attempts to allocate any more, it will be forbidden. Use this field to defensively impose a
-		// strict limit to an application to keep it from going rogue with its memory usage beyond some
-		// undesired limit. The absolute maximum that is possible is one memory page short of 2GB.
-		var MAX_MEMORY_64BIT = Infinity;
-		var MAX_MEMORY_32BIT = 512*1024*1024;
-		var MAX_MEMORY = Math.min(alignPageUp(heuristic64BitBrowser ? MAX_MEMORY_64BIT : MAX_MEMORY_32BIT), 2048 * 1024 * 1024 - pageSize);
-
-		// The application needs at least this much memory to run. If the browser can't provide this,
-		// the page startup should be aborted in an out-of-memory exception.
-		var MIN_MEMORY = Math.min(alignPageUp(32 * 1024 * 1024), MAX_MEMORY);
-
-		// As a hint to the implementation, the application would prefer to reserve this much address
-		// space at startup, and the browser will attempt its best to satisfy this. If this is not
-		// possible, the browser will attempt to allocate anything as close to this IDEAL amount as
-		// possible, but at least MIN bytes.
-		var IDEAL_MEMORY_64BIT = %HEAPSIZE%;
-		var IDEAL_MEMORY_32BIT = %HEAPSIZE%;
-		var IDEAL_MEMORY = Math.min(Math.max(alignPageUp(heuristic64BitBrowser ? IDEAL_MEMORY_64BIT : IDEAL_MEMORY_32BIT), MIN_MEMORY), MAX_MEMORY);
-
-		// If true, assume the application will have most of its memory allocation pressure inside the
-		// application heap, so reserve address space there up front. If false, assume that the memory
-		// allocation pressure is outside the heap, so avoid reserving memory up front, until needed.
-		var RESERVE_MAX_64BIT = true;
-		var RESERVE_MAX_32BIT = false;
-		var RESERVE_MAX = heuristic64BitBrowser ? RESERVE_MAX_64BIT : RESERVE_MAX_32BIT;
-
-		function MB(x) { return (x/1024/1024) + 'MB'; }
-		function allocateHeap() {
-			// Try to get as much memory close to IDEAL, but at least MIN.
-			for(var mem = IDEAL_MEMORY; mem >= MIN_MEMORY; mem -= pageSize) {
-				try {
-					if (RESERVE_MAX)
-						Module['wasmMemory'] = new WebAssembly.Memory({ initial: mem / pageSize, maximum: MAX_MEMORY / pageSize });
-					else
-						Module['wasmMemory'] = new WebAssembly.Memory({ initial: mem / pageSize });
-					Module['buffer'] = Module['wasmMemory'].buffer;
-					if (Module['buffer'].byteLength != mem) throw 'Out of memory';
-					break;
-				} catch(e) { /*nop*/ }
-			}
-			if (!Module['buffer'] || !(Module['buffer'].byteLength >= MIN_MEMORY)) {
-				delete Module['buffer'];
-				throw 'Out of memory';
-			}
-			Module['TOTAL_MEMORY'] = Module['buffer'].byteLength;
-		}
-		allocateHeap();
-		Module['MAX_MEMORY'] = MAX_MEMORY;
-		console.log('Initial memory size: ' + MB(Module['TOTAL_MEMORY']) + ' (MIN_MEMORY: ' + MB(MIN_MEMORY) + ', IDEAL_MEMORY: ' + MB(IDEAL_MEMORY) + ', MAX_MEMORY: ' + MB(MAX_MEMORY) + ', RESERVE_MAX: ' + RESERVE_MAX + ', heuristic64BitBrowser: ' + heuristic64BitBrowser + ', heuristic64BitOS: ' + heuristicIs64Bit('os') + ')');
-	</script>
 	<script src="https://code.jquery.com/jquery-2.1.3.min.js"></script>
 	<script src="https://maxcdn.bootstrapcdn.com/bootstrap/3.3.4/js/bootstrap.min.js"></script>
 	<link href="https://maxcdn.bootstrapcdn.com/bootstrap/3.3.4/css/bootstrap.min.css" rel="stylesheet">
+
+	<link href="%GAME%.css" rel="stylesheet">
 </head>
 <body>
 
-<style type="text/css">
-
-html, body, .container {
-	height: 100%;
-	font-family: "Helvetica Neue", "HelveticaNeue", Helvetica, Arial, sans-serif;
-}
-
-.h4, h4 {
-	margin-top: 1pt;
-	margin-bottom: 1pt;
-	font-size: 10pt;
-}
-
-.container {
-	display: table;
-	vertical-align: middle;
-	border: 0px;
-	border-spacing: 0px;
-}
-
-.glyphicon-spin {
-	animation: spin 2000ms infinite linear;
-}
-
-@keyframes spin {
-	0% {
-		transform: rotate(0deg);
-	}
-	100% {
-		transform: rotate(359deg);
-	}
-}
-
-@-webkit-keyframes spin {
-	0% {
-		transform: rotate(0deg);
-	}
-	100% {
-		transform: rotate(359deg);
-	}
-}
-
-.wrapper {
-	position: relative;
-	margin: 1em auto 10px auto;
-	text-align: center;
-	min-width: 640px;
-	width: 100%;
-	height: 480px; /* initial height, will be dynamically adjusted at runtime */
-	max-width: 95%;
-	display: block;
-	align-items: center;
-	position: relative;
-	text-align: center;
-	justify-content: center;
-}
-
- .emscripten {
-	padding-right: 0;
-	margin-left: auto;
-	margin-right: auto;
-	display: block;
-	display: -webkit-box;
-	display: -moz-box;
-	display: box;
-
-	-webkit-box-align: center;
-	-moz-box-align: center;
-	box-align: center;
-
-	-webkit-box-pack: center;
-	-moz-box-pack: center;
-	box-pack: center;
-}
-
-#canvas:not([fullscreen]) {
-	padding-right: 0;
-	margin-left: auto;
-	margin-right: auto;
-	width: 100%;
-}
-
-.texthalf {
-	height: 37%;
-	border: 0px;
-	padding: 0px;
-	overflow-y: scroll;
-	font-size: 2em;
-}
-
-.buttonarea {
-	min-height: 3%;
-	border-top: 0px;
-	border-bottom: 0px;
-	padding: 0px;
-	margin-right: 0px;
-	margin-top: 0px;
-	margin-bottom: 0px;
- }
-
-.btn { padding: 0px; text-align: center; min-width: 150px }
-.progress { background: rgba(245, 245, 245, 1); border: 0px solid rgba(245, 245, 245, 1); border-radius: 0px; height: 4px; }
-.progress-bar-custom { background: rgba(153, 153, 153, 1); }
-.centered-axis-xy {
-	position: absolute;
-	left: 50%;
-	top: 50%;
-	transform: translate(-50%,-50%);
-}
-</style>
-
 <div class="wrapper" id="mainarea">
 	<div class="alert alert-warning centered-axis-xy" style="min-height: 20px; display:none;" role="alert" id="compilingmessage">
 		<div id='loadTasks'> </div>
@@ -259,956 +42,6 @@ html, body, .container {
 </div>
 <div class="texthalf text-normal jumbotron " id="logwindow" style='display:none'></div>
 
-<script type="text/javascript">
-
-// combine all parallel downloads into one progress bar.
-var totalDownloadSize = 0;
-var currentDownloadedSize = 0;
-
-// helper functions.
-// http://stackoverflow.com/questions/4750015/regular-expression-to-find-urls-within-a-string
-function getHTMLGetParam(name) {
-	if (name=(new RegExp('[?&]'+encodeURIComponent(name)+'=([^&]*)')).exec(location.search))
-		return decodeURIComponent(name[1]);
-}
-
-var UE4 = {
-	on_fatal: function() {
-		try {
-			UE4.on_fatal = Module.cwrap('on_fatal', null, ['string', 'string']);
-		} catch(e) {
-			UE4.on_fatal = function() {};
-		}
-	},
-};
-
-Module['preinitializedWebGLContext'] = null;
-
-function getGpuInfo() {
-	var gl = Module['preinitializedWebGLContext'];
-	if (!gl) return '(no GL: ' + Module['webGLErrorReason'] + ')';
-
-	var glInfo = '';
-	var debugInfo = gl.getExtension('WEBGL_debug_renderer_info');
-	if (debugInfo) glInfo += gl.getParameter(debugInfo.UNMASKED_VENDOR_WEBGL) + ' ' + gl.getParameter(debugInfo.UNMASKED_RENDERER_WEBGL) + '/';
-	glInfo += gl.getParameter(gl.VENDOR) + ' ' + gl.getParameter(gl.RENDERER);
-	glInfo += ' ' + gl.getParameter(gl.VERSION);
-	glInfo += ', ' + gl.getParameter(gl.SHADING_LANGUAGE_VERSION);
-	if (Module['softwareWebGL']) glInfo += ' (software)';
-	return glInfo;
-}
-
-function detectWebGL() {
-	var canvas = Module['canvas'] || document.createElement("canvas");
-	// If you run into problems with WebGL 2, or for quick testing purposes, you can disable UE4
-	// from using WebGL 2 and revert back to WebGL 1 by setting the following flag to true.
-	var disableWebGL2 = false;
-	if (explicitlyUseWebGL1) {
-		disableWebGL2 = true;
-		console.log('Disabled WebGL 2 as requested by ?webgl1 GET param.');
-	}
-	var names = ["webgl", "experimental-webgl"];
-	if (disableWebGL2) {
-		WebGL2RenderingContext = undefined;
-	} else {
-		names = ["webgl2"].concat(names);
-	}
-	function testError(e) { Module['webGLErrorReason'] = e.statusMessage; };
-	canvas.addEventListener("webglcontextcreationerror", testError, false);
-	try {
-		for(var failIfMajorPerformanceCaveat = 1; failIfMajorPerformanceCaveat >= 0; --failIfMajorPerformanceCaveat) {
-			for(var i in names) {
-				try {
-					var context = canvas.getContext(names[i], {antialias:false,alpha:false,depth:true,stencil:true,failIfMajorPerformanceCaveat:!!failIfMajorPerformanceCaveat});
-					Module['preinitializedWebGLContext'] = context;
-					Module['softwareWebGL'] = !failIfMajorPerformanceCaveat;
-					if (context && typeof context.getParameter == "function") {
-						if (typeof WebGL2RenderingContext !== 'undefined' && context instanceof WebGL2RenderingContext && names[i] == 'webgl2') {
-							return 2;
-						} else {
-							// We were able to precreate only a WebGL 1 context, remove support for WebGL 2 from the rest of the page execution.
-							WebGL2RenderingContext = undefined;
-							return 1;
-						}
-					}
-				} catch(e) { Module['webGLErrorReason'] = e.toString(); }
-			}
-		}
-	} finally {
-		canvas.removeEventListener("webglcontextcreationerror", testError, false);
-	}
-	return 0;
-}
-
-Module['canvas'] = document.getElementById('canvas');
-
-// Canvas scaling mode should be set to one of: 1=STRETCH, 2=ASPECT, or 3=FIXED.
-// This dictates how the canvas size changes when the browser window is resized
-// by dragging from the corner.
-var canvasWindowedScaleMode = %CANVASSCALEMODE%;
-
-// High DPI setting configures whether to match the canvas size 1:1 with
-// the physical pixels on the screen.
-// For background, see https://www.khronos.org/webgl/wiki/HandlingHighDPI
-var canvasWindowedUseHighDpi = true;
-
-// Stores the initial size of the canvas in physical pixel units.
-// If canvasWindowedScaleMode == 3 (FIXED), this size defines the fixed resolution
-//                                          that the app will render to.
-// If canvasWindowedScaleMode == 2 (ASPECT), this size defines only the aspect ratio
-//                                           that the canvas will be constrained to.
-// If canvasWindowedScaleMode == 1 (STRETCH), these size values are ignored.
-var canvasAspectRatioWidth = 1366;
-var canvasAspectRatioHeight = 768;
-
-// Fullscreen scaling mode behavior (export these to Module object for the engine to read)
-// This value is one of:
-// 0=NONE: The same canvas size is kept when entering fullscreen without change.
-// 1=STRETCH: The canvas is resized to the size of the whole screen, potentially changing aspect ratio.
-// 2=ASPECT: The canvas is resized to the size of the whole screen, but retaining current aspect ratio.
-// 3=FIXED: The canvas is centered on screen with a fixed resolution.
-Module['UE4_fullscreenScaleMode'] = 1;//canvasWindowedScaleMode; // BUG: if using FIXED, fullscreen gets some strange padding on margin...
-
-// When entering fullscreen mode, should UE4 engine resize the canvas?
-// 0=No resizing (do it manually in resizeCanvas()), 1=Resize to standard DPI, 2=Resize to highDPI
-Module['UE4_fullscreenCanvasResizeMode'] = canvasWindowedUseHighDpi ? 2/*HIDPI*/ : 1/*Standard DPI*/;
-
-// Specifies how canvas is scaled to fullscreen, if not rendering in 1:1 pixel perfect mode.
-// One of 0=Default, 1=Nearest, 2=Bilinear
-Module['UE4_fullscreenFilteringMode'] = 0;
-
-document.addEventListener('error', function(){document.getElementById('clear_indexeddb').style.display = 'inline-block';}, false);
-
-// Startup task which is run after UE4 engine has launched.
-function postRunEmscripten() {
-	taskFinished(TASK_MAIN);
-	$("#compilingmessage").remove();
-
-	// The default Emscripten provided canvas resizing behavior is not needed,
-	// since we are controlling the canvas sizes here, so stub those functions out.
-	Browser.updateCanvasDimensions = function() {};
-	Browser.setCanvasSize = function() {};
-
-	// If you'd like to configure the initial canvas size to render using the resolution
-	// defined in UE4 DefaultEngine.ini [SystemSettings] r.setRes=WidthxHeight,
-	// uncomment the following two lines before calling resizeCanvas() below:
-
-	// canvasAspectRatioWidth = UE_JSlib.UE_GSystemResolution_ResX();
-	// canvasAspectRatioHeight = UE_JSlib.UE_GSystemResolution_ResY();
-
-	// Configure the size of the canvas and display it.
-	resizeCanvas();
-	Module['canvas'].style.display = 'block';
-
-	// Whenever the browser window size changes, relayout the canvas size on the page.
-	window.addEventListener('resize', resizeCanvas, false);
-	window.addEventListener('orientationchange', resizeCanvas, false);
-
-	// The following is needed if game is within an iframe - main window already has focus...
-	window.focus();
-}
-Module.postRun = [postRunEmscripten];
-
-// The resizeCanvas() function recomputes the canvas size on the page as the user changes
-// the browser window size.
-function resizeCanvas(aboutToEnterFullscreen) {
-	// Configuration variables, feel free to play around with these to tweak.
-	var minimumCanvasHeightCssPixels = 480; // the visible size of the canvas should always be at least this high (in CSS pixels)
-	var minimumCanvasHeightFractionOfBrowserWindowHeight = 0.65; // and also vertically take up this much % of the total browser client area height.
-
-	if (aboutToEnterFullscreen && !aboutToEnterFullscreen.type) { // UE4 engine is calling this function right before entering fullscreen?
-		// If you want to perform specific resolution setup here, do so by setting Module['canvas'].width x Module['canvas'].height now,
-		// and configure Module['UE4_fullscreenXXX'] fields above. Most of the time, the defaults are good, so no need to resize here.
-		// Return true here if you want to abort entering fullscreen mode altogether.
-		return;
-	}
-
-	// The browser called resizeCanvas() to notify that we just entered fullscreen? In that case, we never react, since the strategy is
-	// to always set the canvas size right before entering fullscreen.
-	if (document.fullscreenElement || document.mozFullScreenElement || document.webkitFullscreenElement || document.msFullscreenElement) {
-		return;
-	}
-
-	var mainArea = document.getElementById('mainarea');
-	var canvasRect = mainArea.getBoundingClientRect();
-
-	// Compute the unconstrained size for the div that encloses the canvas, in CSS pixel units.
-	var cssWidth = canvasRect.right - canvasRect.left;
-	var cssHeight = Math.max(minimumCanvasHeightCssPixels, canvasRect.bottom - canvasRect.top, window.innerHeight * minimumCanvasHeightFractionOfBrowserWindowHeight);
-
-	if (canvasWindowedScaleMode == 3/*NONE*/) {
-		// In fixed display mode, render to a statically determined WebGL render target size.
-		var newRenderTargetWidth = canvasAspectRatioWidth;
-		var newRenderTargetHeight = canvasAspectRatioHeight;
-	} else {
-		// Convert unconstrained render target size from CSS to physical pixel units.
-		var newRenderTargetWidth = canvasWindowedUseHighDpi ? (cssWidth * window.devicePixelRatio) : cssWidth;
-		var newRenderTargetHeight = canvasWindowedUseHighDpi ? (cssHeight * window.devicePixelRatio) : cssHeight;
-
-		// Apply aspect ratio constraints, if desired.
-		if (canvasWindowedScaleMode == 2/*ASPECT*/) {
-			if (cssWidth * canvasAspectRatioHeight > canvasAspectRatioWidth * cssHeight) {
-				newRenderTargetWidth = newRenderTargetHeight * canvasAspectRatioWidth / canvasAspectRatioHeight;
-			} else {
-				newRenderTargetHeight = newRenderTargetWidth * canvasAspectRatioHeight / canvasAspectRatioWidth;
-			}
-		}
-
-		// WebGL render target sizes are always full integer pixels in size, so rounding is critical for CSS size computations below.
-		newRenderTargetWidth = Math.round(newRenderTargetWidth);
-		newRenderTargetHeight = Math.round(newRenderTargetHeight);
-	}
-
-	// Very subtle but important behavior is that the size of a DOM element on a web page in CSS pixel units can be a fraction, e.g. on
-	// high DPI scaling displays (CSS pixel units are "virtual" pixels). If the CSS size and physical pixel size of the WebGL canvas do
-	// not correspond to each other 1:1 after window.devicePixelRatio scaling has been applied, the result can look blurry. Therefore always
-	// first compute the WebGL render target size first in physical pixels, and convert that back to CSS pixels so that the CSS pixel size
-	// will perfectly align up and the result look clear without scaling applied.
-	cssWidth = canvasWindowedUseHighDpi ? (newRenderTargetWidth / window.devicePixelRatio) : newRenderTargetWidth;
-	cssHeight = canvasWindowedUseHighDpi ? (newRenderTargetHeight / window.devicePixelRatio) : newRenderTargetHeight;
-
-	Module['canvas'].width = newRenderTargetWidth;
-	Module['canvas'].height = newRenderTargetHeight;
-
-	Module['canvas'].style.width = cssWidth + 'px';
-	Module['canvas'].style.height = mainArea.style.height = cssHeight + 'px';
-
-	// Tell the engine that the web page has changed the size of the WebGL render target on the canvas (Module['canvas'].width/height).
-	// This will update the GL viewport and propagate the change throughout the engine.
-	// If the CSS style size is changed, this function doesn't need to be called.
-	if (UE_JSlib.UE_CanvasSizeChanged) UE_JSlib.UE_CanvasSizeChanged();
-}
-Module['UE4_resizeCanvas'] = resizeCanvas;
-
-Module.arguments = [%UE4CMDLINE%];
-
-// we are serving via a server and it is unreal file server?
-if (location.host != "" && getHTMLGetParam("cookonthefly") == "true") {
-	Module.arguments.push("'-filehostIp=" + location.protocol + "//" + location.host + "'");
-}
-
-function addLog(info, color) {
-	$("#logwindow").append("<h4><small>" + info + " </small></h4>");
-}
-Module.print = addLog;
-
-Module.printErr = function(text) {
-	console.error(text);
-};
-
-// Module.locateFile() routes asset downloads to either gzip compressed or uncompressed assets.
-Module.locateFile = function(name) {
-	var serveGzipped = serveCompressedAssets;
-	// When serving from file:// URLs, don't read .gz compressed files, because these files can't be transparently uncompressed.
-	var isFileProtocol = name.indexOf('file://') != -1 || location.protocol.indexOf('file') != -1;
-	if (isFileProtocol) {
-		if (!Module['shownFileProtocolWarning']) {
-			showWarningRibbon('Attempting to load the page via the "file://" protocol. This only works in Firefox, and even there only when not using compression, so attempting to load uncompressed assets. Please host the page on a web server and visit it via a "http://" URL.');
-			Module['shownFileProtocolWarning'] = true;
-		}
-		serveGzipped = false;
-	}
-
-	// uncompressing very large gzip files may slow down startup times.
-//	if (!dataFileIsGzipCompressed && name.split('.').slice(-1)[0] == 'data') serveGzipped = false;
-
-	return serveGzipped ? (name + 'gz') : name;
-};
-
-Module.getPreloadedPackage = function(remotePackageName, remotePackageSize) {
-	return Module['preloadedPackages'] ? Module['preloadedPackages'][remotePackageName] : null;
-}
-
-// Asynchronously appends the given script code to DOM. This is to ensure that
-// browsers parse and compile the JS code parallel to all other execution.
-function addScriptToDom(scriptCode) {
-	return new Promise(function(resolve, reject) {
-		var script = document.createElement('script');
-		var blob = (scriptCode instanceof Blob) ? scriptCode : new Blob([scriptCode], { type: 'text/javascript' });
-		var objectUrl = URL.createObjectURL(blob);
-		script.src = objectUrl;
-		script.onload = function() {
-			script.onload = script.onerror = null; // Remove these onload and onerror handlers, because these capture the inputs to the Promise and the input function, which would leak a lot of memory!
-			URL.revokeObjectURL(objectUrl); // Free up the blob. Note that for debugging purposes, this can be useful to comment out to be able to read the sources in debugger.
-			resolve();
-		}
-		script.onerror = function(e) {
-			script.onload = script.onerror = null; // Remove these onload and onerror handlers, because these capture the inputs to the Promise and the input function, which would leak a lot of memory!
-			URL.revokeObjectURL(objectUrl);
-			console.error('script failed to add to dom: ' + e);
-			console.error(scriptCode);
-			console.error(e);
-			// The onerror event sends a DOM Level 3 event error object, which does not seem to have any kind of human readable error reason (https://developer.mozilla.org/en-US/docs/Web/Events/error)
-			// There is another error event object at https://developer.mozilla.org/en-US/docs/Web/API/ErrorEvent, which would have an error reason. Perhaps that error event might sometimes be fired,
-			// but if not, guess that the error reason was an OOM, since we are dealing with large .js files.
-			reject(e.message || "(out of memory?)");
-		}
-		document.body.appendChild(script);
-	});
-}
-
-var TASK_DOWNLOADING = 0;
-var TASK_COMPILING = 1;
-var TASK_SHADERS = 2;
-var TASK_MAIN = 3;
-var loadTasks = [ 'Downloading', 'Compiling WebAssembly', 'Building shaders', 'Launching engine'];
-
-function taskProgress(taskId, progress) {
-	var c = document.getElementById('compilingmessage');
-	if (c) c.style.display = 'block';
-	else return;
-	var l = document.getElementById('load_' + taskId);
-	if (!l) {
-		var tasks = document.getElementById('loadTasks');
-		if (!tasks) return;
-		l = document.createElement('div');
-		l.innerHTML = '<span id="icon_' + taskId + '" class="glyphicon glyphicon-refresh glyphicon-spin"></span>  <span id="load_' + taskId + '"></span>';
-		tasks.appendChild(l);
-		l = document.getElementById('load_' + taskId);
-	}
-	if (!l.startTime) l.startTime = performance.now();
-	var text = loadTasks[taskId];
-	if (progress && progress.total) {
-		text += ': ' + (progress.currentShow || progress.current) + '/' + (progress.totalShow || progress.total) + ' (' + (progress.current * 100 / progress.total).toFixed(0) + '%)';
-	} else {
-		text += '...';
-	}
-	l.innerHTML = text;
-}
-
-function showErrorDialog(errorText) {
-	if ( errorText.indexOf('SyntaxError: ') != -1 ) { // this may be due to caching issue -- otherwise, compile time would have caught this
-		errorText = "NOTE: attempting to flush cache and force reload...<br>Please standby...";
-		setTimeout(function() {
-			location.reload(true);
-		}, 2000); // 2 seconds
-	}
-	console.error('error: ' + errorText);
-	var existingErrorDialog = document.getElementById('errorDialog');
-	if (existingErrorDialog) {
-		existingErrorDialog.innerHTML += '<br>' + errorText;
-	} else {
-		$('#mainarea').empty();
-		$('#mainarea').append('<div class="alert alert-danger centered-axis-xy" style ="min-height: 10pt" role="alert" id="errorDialog">' + errorText + '</div></div>');
-	}
-}
-
-function showWarningRibbon(warningText) {
-	var existingWarningDialog = document.getElementById('warningDialog');
-	if (existingWarningDialog) {
-		existingWarningDialog.innerHTML += '<br>' + warningText;
-	} else {
-		$('#buttonrow').prepend('<div class="alert alert-warning centered-axis-x" role="warning" id="warningDialog" style="padding-top:5px; padding-bottom: 5px">' + warningText + '</div></div>');
-	}
-}
-
-function taskFinished(taskId, error) {
-	var l = document.getElementById('load_' + taskId);
-	var icon = document.getElementById('icon_' + taskId);
-	if (l && icon) {
-		var totalTime = performance.now() - l.startTime;
-		if (!error) {
-			l.innerHTML = loadTasks[taskId] + ' (' + (totalTime/1000).toFixed(2) + 's)';
-			icon.className = 'glyphicon glyphicon-ok';
-		}
-		else {
-			l.innerHTML = loadTasks[taskId] + ': FAILED! ' + error;
-			icon.className = 'glyphicon glyphicon-remove';
-
-			showErrorDialog(loadTasks[taskId] + ' failed: <br> ' + error);
-		}
-	}
-}
-
-window.onerror = function(e) {
-	e = e.toString();
-	if (e.toLowerCase().indexOf('memory') != -1) {
-		e += '<br>';
-		if (!heuristic64BitBrowser) e += ' Try running in a 64-bit browser to resolve.';
-	}
-	showErrorDialog(e);
-}
-
-function formatBytes(bytes) {
-	if (bytes >= 1024*1024*1024) return (bytes / (1024*1024*1024)).toFixed(1) + ' GB';
-	if (bytes >= 1024*1024) return (bytes / (1024*1024)).toFixed(0) + ' MB';
-	if (bytes >= 1024) return (bytes / 1024).toFixed(1) + ' KB';
-	return bytes + ' B';
-}
-
-function reportDataBytesStoredInIndexedDB(deltaBytes) {
-	if (deltaBytes === null) Module['dataBytesStoredInIndexedDB'] = 0; // call with deltaBytes == null to report that DB was cleared.
-	else Module['dataBytesStoredInIndexedDB'] += deltaBytes;
-	document.getElementById('clear_indexeddb').innerText = 'Clear IndexedDB (' + formatBytes(Module['dataBytesStoredInIndexedDB']) + ')';
-}
-
-function reportDownloadProgress(url, downloadedBytes, totalBytes, finished) {
-	Module['assetDownloadProgress'][url] = {
-		current: downloadedBytes,
-		total: totalBytes,
-		finished: finished
-	};
-	var aggregated = {
-		current: 0,
-		total: 0,
-		finished: true
-	};
-	for(var i in Module['assetDownloadProgress']) {
-		aggregated.current += Module['assetDownloadProgress'][i].current;
-		aggregated.total += Module['assetDownloadProgress'][i].total;
-		aggregated.finished = aggregated.finished && Module['assetDownloadProgress'][i].finished;
-	}
-
-	aggregated.currentShow = formatBytes(aggregated.current);
-	aggregated.totalShow = formatBytes(aggregated.total);
-
-	if (aggregated.finished) taskFinished(TASK_DOWNLOADING);
-	else taskProgress(TASK_DOWNLOADING, aggregated);
-}
-
-function download(url, responseType) {
-	return new Promise(function(resolve, reject) {
-		var xhr = new XMLHttpRequest();
-		xhr.open('GET', url, true);
-		xhr.responseType = responseType || 'blob';
-		reportDownloadProgress(url, 0, 1);
-		xhr.onload = function() {
-			if (xhr.status == 0 || (xhr.status >= 200 && xhr.status < 300)) {
-				var len = xhr.response.size || xhr.response.byteLength;
-				reportDownloadProgress(url, len, len, true);
-				resolve(xhr.response);
-			} else {
-				taskFinished(TASK_DOWNLOADING, 'HTTP error ' + (xhr.status || 404) + ' ' + xhr.statusText + ' on file ' + url);
-				reject({
-					status: xhr.status,
-					statusText: xhr.statusText
-				});
-			}
-		};
-		xhr.onprogress = function(p) {
-			if (p.lengthComputable) reportDownloadProgress(url, p.loaded, p.total);
-		};
-		xhr.onerror = function(e) {
-			var isFileProtocol = url.indexOf('file://') == 0 || location.protocol.indexOf('file') != -1;
-			if (isFileProtocol) taskFinished(TASK_DOWNLOADING, 'HTTP error ' + (xhr.status || 404) + ' ' + xhr.statusText + ' on file ' + url +'<br>Try using a web server to avoid loading via a "file://" URL.'); // Convert the most common source of errors to a more friendly message format.
-			else taskFinished(TASK_DOWNLOADING, 'HTTP error ' + (xhr.status || 404) + ' ' + xhr.statusText + ' on file ' + url);
-			reject({
-				status: xhr.status || 404,
-				statusText: xhr.statusText
-			});
-		};
-		xhr.onreadystatechange = function() {
-			if (xhr.readyState >= xhr.HEADERS_RECEIVED) {
-				if (url.endsWith('gz') && (xhr.status == 0 || xhr.status == 200)) {
-					if (xhr.getResponseHeader('Content-Encoding') != 'gzip') {
-						// A fallback is to set serveCompressedAssets = false to serve uncompressed assets instead, but that is not really recommended for production use, since gzip compression shrinks
-						// download sizes so dramatically that omitting it for production is not a good idea.
-						taskFinished(TASK_DOWNLOADING, 'Downloaded a compressed file ' + url + ' without the necessary HTTP response header "Content-Encoding: gzip" specified!<br>Please configure gzip compression on this asset on the web server to serve gzipped assets!');
-						xhr.onload = xhr.onprogress = xhr.onerror = xhr.onreadystatechange = null; // Abandon tracking events from this XHR further.
-						xhr.abort();
-						return reject({
-							status: 406,
-							statusText: 'Not Acceptable'
-						});
-					}
-
-					// After enabling Content-Encoding: gzip, make sure that the appropriate MIME type is being used for the asset, i.e. the MIME
-					// type should be that of the uncompressed asset, and not the MIME type of the compression method that was used.
-					if (xhr.getResponseHeader('Content-Type').toLowerCase().indexOf('zip') != -1) {
-						function expectedMimeType(url) {
-							if (url.indexOf('.wasm') != -1) return 'application/wasm';
-							if (url.indexOf('.js') != -1) return 'application/javascript';
-							return 'application/octet-stream';
-						}
-						taskFinished(TASK_DOWNLOADING, 'Downloaded a compressed file ' + url + ' with incorrect HTTP response header "Content-Type: ' + xhr.getResponseHeader('Content-Type') + '"!<br>Please set the MIME type of the asset to "' + expectedMimeType(url) + '".');
-						xhr.onload = xhr.onprogress = xhr.onerror = xhr.onreadystatechange = null; // Abandon tracking events from this XHR further.
-						xhr.abort();
-						return reject({
-							status: 406,
-							statusText: 'Not Acceptable'
-						});
-					}
-				}
-			}
-		}
-		xhr.send(null);
-	});
-}
-
-function getIDBRequestErrorString(req) {
-	try { return req.error ? ('IndexedDB ' + req.error.name + ': ' + req.error.message) : req.result;
-	} catch(ex) { return null; }
-}
-
-function deleteIndexedDBStorage(dbName, onsuccess, onerror, onblocked) {
-	var idb = window.indexedDB || window.mozIndexedDB || window.webkitIndexedDB || window.msIndexedDB;
-	if (Module['dbInstance']) Module['dbInstance'].close();
-	if (!dbName) dbName = Module['UE4_indexedDBName'];
-	var req = idb.deleteDatabase(dbName);
-	req.onsuccess = function() { console.log('Deleted IndexedDB storage ' + dbName + '!'); reportDataBytesStoredInIndexedDB(null); if (onsuccess) onsuccess(); }
-	req.onerror = function(evt) {
-		var errorString = getIDBRequestErrorString(req);
-		console.error('Failed to delete IndexedDB storage ' + dbName + ', ' + errorString);
-		evt.preventDefault();
-		if (onerror) onerror(errorString);
-	};
-	req.onblocked = function(evt) {
-		var errorString = getIDBRequestErrorString(req);
-		console.error('Failed to delete IndexedDB storage ' + dbName + ', DB was blocked! ' + errorString);
-		evt.preventDefault();
-		if (onblocked) onblocked(errorString);
-	}
-}
-
-var enableReadFromIndexedDB = (location.search.indexOf('noidbread') == -1);
-var enableWriteToIndexedDB = enableReadFromIndexedDB && (location.search.indexOf('noidbwrite') == -1);
-%DISABLE_INDEXEDDB%
-
-if (!enableReadFromIndexedDB) showWarningRibbon('Running with IndexedDB access disabled.');
-else if (!enableWriteToIndexedDB) showWarningRibbon('Running in read-only IndexedDB access mode.');
-
-function storeToIndexedDB(db, key, value) {
-	return new Promise(function(resolve, reject) {
-		if (!enableWriteToIndexedDB) return reject('storeToIndexedDB: IndexedDB writes disabled by "?noidbwrite" option');
-		function fail(e) {
-			console.error('Failed to store file ' + key + ' to IndexedDB storage! error: ' + e);
-			if (!Module['idberrorShown']) {
-				showWarningRibbon('Failed to store file ' + key + ' to IndexedDB, error: ' + e);
-				Module['idberrorShown'] = true;
-			}
-			return reject(e);
-		}
-		if (!db) return fail('IndexedDB not available!');
-		if (location.protocol.indexOf('file') != -1) return reject('Loading via file://, skipping caching to IndexedDB');
-
-		try {
-			var transaction = db.transaction(['FILES'], 'readwrite');
-			var packages = transaction.objectStore('FILES');
-			var putRequest = packages.put(value, "file/" + Module.key + '/' + key);
-			putRequest.onsuccess = function(evt) {
-				if (value.byteLength || value.length) reportDataBytesStoredInIndexedDB(value.size || value.byteLength || value.length);
-				resolve(key);
-			};
-			putRequest.onerror = function(evt) {
-				var errorString = getIDBRequestErrorString(putRequest) || ('IndexedDB request error: ' + evt);
-				evt.preventDefault();
-				fail(errorString);
-			};
-		} catch(e) {
-			fail(e);
-		}
-	});
-}
-
-function fetchFromIndexedDB(db, key) {
-	return new Promise(function(resolve, reject) {
-		if (!enableReadFromIndexedDB) return reject('fetchFromIndexedDB: IndexedDB reads disabled by "?noidbread" option');
-		function fail(e) {
-			console.error('Failed to read file ' + key + ' from IndexedDB storage! error:');
-			console.error(e);
-			if (!Module['idberrorShown']) {
-				showWarningRibbon('Failed to read file ' + key + ' from IndexedDB, error: ' + e);
-				Module['idberrorShown'] = true;
-			}
-			return reject(e);
-		}
-		if (!db) return fail('IndexedDB not available!');
-		try {
-			var transaction = db.transaction(['FILES'], 'readonly');
-			var packages = transaction.objectStore('FILES');
-			var getRequest = packages.get("file/" + Module.key + '/' + key);
-			getRequest.onsuccess = function(evt) {
-				if (evt.target.result) {
-					var len = evt.target.result.size || evt.target.result.byteLength || evt.target.result.length;
-					if (len) reportDataBytesStoredInIndexedDB(len);
-					resolve(evt.target.result);
-				} else {
-					// Succeeded to load, but the load came back with the value of undefined, treat that as an error since we never store undefined in db.
-					reject();
-				}
-			};
-			getRequest.onerror = function(evt) {
-				var errorString = getIDBRequestErrorString(getRequest) || ('IndexedDB.get request error: ' + evt);
-				evt.preventDefault();
-				fail(errorString);
-			};
-		} catch(e) {
-			fail(e);
-		}
-	});
-}
-
-function fetchOrDownloadAndStore(db, url, responseType) {
-	return new Promise(function(resolve, reject) {
-		fetchFromIndexedDB(db, url)
-		.then(function(data) { return resolve(data); })
-		.catch(function(error) {
-			return download(url, responseType)
-			.then(function(data) {
-				// Treat IDB store as separate operation that's not part of the Promise chain.
-				/*return*/ storeToIndexedDB(db, url, data)
-				.then(function() { return resolve(data); })
-				.catch(function(error) {
-					console.error('Failed to store download to IndexedDB! ' + error);
-					return resolve(data); // succeeded download, but failed to store - ignore failure in that case and just proceed to run by calling the success handler.
-				})
-			})
-			.catch(function(error) { return reject(error); })
-		});
-	});
-}
-
-function openIndexedDB(dbName, dbVersion) {
-	return new Promise(function(resolve, reject) {
-		if (!enableReadFromIndexedDB) return reject('openIndexedDB: IndexedDB disabled by "?noidbread" option');
-		try {
-			var idb = window.indexedDB || window.mozIndexedDB || window.webkitIndexedDB || window.msIndexedDB;
-			var openRequest = idb.open(dbName, dbVersion);
-		} catch(e) { return reject(e); }
-
-		openRequest.onupgradeneeded = function(evt) {
-			var db = evt.target.result;
-			if (db.objectStoreNames.contains('FILES')) db.deleteObjectStore('FILES');
-			db.createObjectStore('FILES');
-		};
-		openRequest.onsuccess = function(evt) {
-			resolve(evt.target.result);
-		};
-		openRequest.onerror = function(evt) {
-			var errorString = getIDBRequestErrorString(openRequest) || ('IndexedDB request error: ' + evt);
-			evt.preventDefault();
-			reject(errorString);
-		};
-	});
-}
-
-Module['instantiateWasm'] = function(info, receiveInstance) {
-	Module['wasmDownloadAction'].then(function(downloadResults) {
-		taskProgress(TASK_COMPILING);
-		var wasmInstantiate = WebAssembly.instantiate(downloadResults.wasmModule || new Uint8Array(downloadResults.wasmBytes), info);
-		return wasmInstantiate.then(function(output) {
-			var instance = output.instance || output;
-			var module = output.module;
-			taskFinished(TASK_COMPILING);
-			Module['wasmInstantiateActionResolve'](instance);
-			receiveInstance(instance);
-
-			// After a successful instantiation, attempt to save the compiled Wasm Module object to IndexedDB.
-			if (!downloadResults.fromIndexedDB) {
-				storeToIndexedDB(downloadResults.db, 'wasmModule', module).catch(function() {
-					// If the browser did not support storing Wasm Modules to IndexedDB, try to store the Wasm instance instead.
-					return storeToIndexedDB(downloadResults.db, 'wasmBytes', downloadResults.wasmBytes);
-				});
-			}
-		});
-	}).catch(function(error) {
-		$ ('#mainarea').empty();
-		$ ('#mainarea').append('<div class="alert alert-danger centered-axis-xy" style ="min-height: 10pt" role="alert">WebAssembly instantiation failed: <br> ' + error + '</div></div>');
-	});
-	return {};
-}
-
-// Given a blob, asynchronously reads the byte contents of that blob to an arraybuffer and returns it as a Promise.
-function readBlobToArrayBuffer(blob) {
-	return new Promise(function(resolve, reject) {
-		var fileReader = new FileReader();
-		fileReader.onload = function() { resolve(this.result); }
-		fileReader.onerror = function(e) { reject(e); }
-		fileReader.readAsArrayBuffer(blob);
-	});
-}
-
-function compileShadersFromJson(jsonData) {
-	var shaderPrograms = [];
-	if (jsonData instanceof ArrayBuffer) jsonData = new TextDecoder('utf-8').decode(new DataView(jsonData));
-	var programsDict = JSON.parse(jsonData);
-	for(var i in programsDict) {
-		shaderPrograms.push(programsDict[i]);
-	}
-
-	var gl = Module['preinitializedWebGLContext'];
-
-	Module['precompiledShaders'] = [];
-	Module['precompiledPrograms'] = [];
-
-	Module['glIDCounter'] = 1;
-	Module['precompiledUniforms'] = [null];
-
-	var promise = new Promise(function(resolve, reject) {
-		var nextProgramToBuild = 0;
-		function buildProgram() {
-			if (nextProgramToBuild >= shaderPrograms.length) {
-				taskFinished(TASK_SHADERS);
-				return resolve();
-			}
-			var p = shaderPrograms[nextProgramToBuild++];
-			taskProgress(TASK_SHADERS, {current: nextProgramToBuild, total: shaderPrograms.length });
-			var program = gl.createProgram();
-
-			function lineNumberize(str) {
-				str = str.split('\n');
-				for(var i = 0; i < str.length; ++i) str[i] = (i<9?' ':'') + (i<99?' ':'') + (i+1) + ': ' + str[i];
-				return str.join('\n');
-			}
-
-			var vs = gl.createShader(gl.VERTEX_SHADER);
-			gl.shaderSource(vs, p.vs);
-			gl.compileShader(vs);
-			var success = gl.getShaderParameter(vs, gl.COMPILE_STATUS);
-			var compileLog = gl.getShaderInfoLog(vs);
-			if (compileLog) compileLog = compileLog.trim();
-			if (compileLog) console.error('Compiling vertex shader: ' + lineNumberize(p.vs));
-			if (!success) console.error('Vertex shader compilation failed!');
-			if (compileLog) console.error('Compilation log: ' + compileLog);
-			if (!success) return reject('Vertex shader compilation failed: ' + compileLog);
-			gl.attachShader(program, vs);
-
-			Module['precompiledShaders'].push({
-				vs: p.vs,
-				shader: vs,
-				program: program
-			});
-
-			var fs = gl.createShader(gl.FRAGMENT_SHADER);
-			gl.shaderSource(fs, p.fs);
-			gl.compileShader(fs);
-			var success = gl.getShaderParameter(fs, gl.COMPILE_STATUS);
-			var compileLog = gl.getShaderInfoLog(fs);
-			if (compileLog) compileLog = compileLog.trim();
-			if (compileLog) console.error('Compiling fragment shader: ' + lineNumberize(p.fs));
-			if (!success) console.error('Fragment shader compilation failed!');
-			if (compileLog) console.error('Compilation log: ' + compileLog);
-			if (!success) return reject('Fragment shader compilation failed: ' + compileLog);
-			gl.attachShader(program, fs);
-
-			Module['precompiledShaders'].push({
-				fs: p.fs,
-				shader: fs,
-				program: program
-			});
-
-			for(var name in p.attribs) {
-				gl.bindAttribLocation(program, p.attribs[name], name);
-			}
-			gl.linkProgram(program);
-			var success = gl.getProgramParameter(program, gl.LINK_STATUS);
-			var linkLog = gl.getProgramInfoLog(program);
-			if (linkLog) linkLog = linkLog.trim();
-			if (linkLog) console.error('Linking shader program, vs: \n' + lineNumberize(p.vs) + ', \n fs:\n' + lineNumberize(p.fs));
-			if (!success) console.error('Shader program linking failed!');
-			if (linkLog) console.error('Link log: ' + linkLog);
-			if (!success) return reject('Shader linking failed: ' + linkLog);
-
-			var ptable = {
-				uniforms: {},
-				maxUniformLength: 0,
-				maxAttributeLength: -1,
-				maxUniformBlockNameLength: -1
-			};
-			var GLctx = gl;
-			var utable = ptable.uniforms;
-				var numUniforms = GLctx.getProgramParameter(program, GLctx.ACTIVE_UNIFORMS);
-				for (var i = 0; i < numUniforms; ++i) {
-				var u = GLctx.getActiveUniform(program, i);
-					var name = u.name;
-					ptable.maxUniformLength = Math.max(ptable.maxUniformLength, name.length + 1);
-					if (name.indexOf("]", name.length - 1) !== -1) {
-					var ls = name.lastIndexOf("[");
-					name = name.slice(0, ls);
-					}
-					var loc = GLctx.getUniformLocation(program, name);
-					var id = Module['glIDCounter']++;
-					utable[name] = [ u.size, id ];
-					Module['precompiledUniforms'].push(loc);
-					if (Module['precompiledUniforms'].length != Module['glIDCounter']) throw 'uniforms array not in sync! ' + Module['precompiledUniforms'].length + ', ' + Module['glIDCounter'];
-					for (var j = 1; j < u.size; ++j) {
-					var n = name + "[" + j + "]";
-					loc = GLctx.getUniformLocation(program, n);
-					id = Module['glIDCounter']++;
-					Module['precompiledUniforms'].push(loc);
-						if (Module['precompiledUniforms'].length != Module['glIDCounter']) throw 'uniforms array not in sync! ' + Module['precompiledUniforms'].length + ', ' + Module['glIDCounter'];
-					}
-				}
-
-			var e = gl.getError();
-			if (e) {
-				console.error('Precompiling shaders got GL error: ' + e);
-				return reject('Precompiling shaders got GL error: ' + e);
-			}
-			Module['precompiledPrograms'].push({
-				program: program,
-				programInfos: ptable,
-				vs: p.vs,
-				fs: p.fs
-			});
-				setTimeout(buildProgram, 0);
-			}
-		setTimeout(buildProgram, 0);
-	})
-
-	return promise;
-}
-
-$(document).ready(function() {
-	if (!supportsWasm) {
-		showErrorDialog('Your browser does not support WebAssembly. Please try updating to latest 64-bit browser that supports WebAssembly.<br>Current user agent: ' + navigator.userAgent);
-		return;
-	}
-
-	if (!Module['buffer'] && allocateHeapUpFront) {
-		showErrorDialog('Failed to allocate ' + MB(MIN_MEMORY) + ' of linear memory for the ' + 'WebAssembly' + ' heap!');
-		return;
-	}
-	// check for webgl and cache it for later (UE_BrowserWebGLVersion() reads this)
-	Module['WEBGL_VERSION'] = detectWebGL();
-	console.log(getGpuInfo());
-	if (!Module['WEBGL_VERSION'] || Module['WEBGL_VERSION'] < requiredWebGLVersion) {
-		showErrorDialog('Your browser does not support WebGL ' + requiredWebGLVersion + '<br>Error reason: ' + (Module['webGLErrorReason'] || 'Unknown') + '. Try updating your browser and/or graphics card drivers.<br>Current renderer: ' + getGpuInfo());
-		return;
-	}
-
-	function shouldBrowserSupportWebGL2() {
-		var match = window.navigator.userAgent.match(/Firefox\/([0-9]+)\./);
-		if (match) return parseInt(match[1]) >= 51;
-	}
-
-	if (Module['WEBGL_VERSION'] < 2 && !explicitlyUseWebGL1) {
-		if (shouldBrowserSupportWebGL2()) {
-			showWarningRibbon('Your GPU does not support WebGL 2. This affects graphics performance and quality. Please try updating your graphics driver and/or browser to latest version.<br>Error reason: ' + (Module['webGLErrorReason'] || 'Unknown') + '<br>Current renderer: ' + getGpuInfo());
-		} else {
-			showWarningRibbon('The current browser does not support WebGL 2. This affects graphics performance and quality.<br>Please try updating your browser (and/or video drivers).  NOTE: old hardware might have been blacklisted by this browser -- you may need to use a different browser.<br>Error reason: ' + (Module['webGLErrorReason'] || 'Unknown') + '<br>Current renderer: ' + getGpuInfo());
-		}
-	}
-
-	if (!heuristicIs64Bit('browser')) {
-		if (heuristicIs64Bit('os')) {
-			showWarningRibbon('It looks like you are running a 32-bit browser on a 64-bit operating system. This can dramatically affect performance and risk running out of memory on large applications. Try updating to a 64-bit browser for an optimized experience.');
-		} else {
-			showWarningRibbon('It looks like your computer hardware is 32-bit. This can dramatically affect performance.');
-		}
-	}
-
-	// The following WebGL 1.0 extensions are available in core WebGL 2.0 specification, so they are no longer shown in the extensions list.
-	var webGLExtensionsInCoreWebGL2 = ['ANGLE_instanced_arrays','EXT_blend_minmax','EXT_color_buffer_half_float','EXT_frag_depth','EXT_sRGB','EXT_shader_texture_lod','OES_element_index_uint','OES_standard_derivatives','OES_texture_float','OES_texture_half_float','OES_texture_half_float_linear','OES_vertex_array_object','WEBGL_color_buffer_float','WEBGL_depth_texture','WEBGL_draw_buffers'];
-
-	var supportedWebGLExtensions = Module['preinitializedWebGLContext'].getSupportedExtensions();
-	if (Module['WEBGL_VERSION'] >= 2) supportedWebGLExtensions = supportedWebGLExtensions.concat(webGLExtensionsInCoreWebGL2);
-
-	// The following WebGL extensions are required by UE4/this project, and it cannot run without.
-	var requiredWebGLExtensions = []; // TODO: List WebGL extensions here that the demo needs and can't run without.
-	for(var i in requiredWebGLExtensions) {
-		if (supportedWebGLExtensions.indexOf(requiredWebGLExtensions[i]) == -1) {
-			showErrorDialog('Your browser does not support WebGL extension ' + requiredWebGLExtensions[i] + ', which is required to run this page!');
-		}
-	}
-
-	// The following WebGL extensions would be preferred to exist for best features/performance, but are not strictly needed and UE4 can fall back if not available.
-	var preferredToHaveWebGLExtensions = [// The following are core in WebGL 2:
-	                                      'ANGLE_instanced_arrays', // UE4 uses instanced rendering where possible, but can fallback to noninstanced.
-	                                      'EXT_color_buffer_half_float',
-	                                      'EXT_sRGB',
-	                                      'EXT_shader_texture_lod', // textureLod() is needed for correct reflections, without this reflection shaders are missing and render out black.
-	                                      'OES_standard_derivatives',
-	                                      'OES_texture_half_float',
-	                                      'OES_texture_half_float_linear',
-	                                      'OES_vertex_array_object',
-	                                      'WEBGL_color_buffer_float',
-	                                      'WEBGL_depth_texture',
-	                                      'WEBGL_draw_buffers',
-
-	                                      // These are still extensions in WebGL 2:
-	                                      'OES_texture_float',
-	                                      'WEBGL_compressed_texture_s3tc',
-	                                      'EXT_texture_filter_anisotropic'
-	];
-	var unsupportedWebGLExtensions = [];
-	for(var i in preferredToHaveWebGLExtensions) {
-		if (supportedWebGLExtensions.indexOf(preferredToHaveWebGLExtensions[i]) == -1) {
-			unsupportedWebGLExtensions.push(preferredToHaveWebGLExtensions[i]);
-		}
-	}
-	if (unsupportedWebGLExtensions.length > 1) {
-		showWarningRibbon('Your browser or graphics card does not support the following WebGL extensions: ' + unsupportedWebGLExtensions.join(', ') + '. This can impact UE4 graphics performance and quality.');
-	} else if (unsupportedWebGLExtensions.length == 1) {
-		showWarningRibbon('Your browser or graphics card does not support the WebGL extension ' + unsupportedWebGLExtensions[0] + '. This can impact UE4 graphics performance and quality.');
-	}
-
-	function withIndexedDB(db) {
-		Module['dbInstance'] = db;
-
-		var mainCompiledCode = fetchFromIndexedDB(db, 'wasmModule').then(function(wasmModule) {
-			return { db: db, wasmModule: wasmModule, fromIndexedDB: true };
-		}).catch(function() {
-			return fetchFromIndexedDB(db, 'wasmBytes').then(function(wasmBytes) {
-				return { db: db, wasmBytes: wasmBytes, fromIndexedDB: true };
-			});
-		}).catch(function() {
-			return download(Module.locateFile('%CONFIG%.wasm'), 'arraybuffer').then(function(wasmBytes) {
-				return { db: db, wasmBytes: wasmBytes, fromIndexedDB: false };
-			});
-		});
-		Module['wasmDownloadAction'] = mainCompiledCode;
-		var compiledCodeInstantiateAction = new Promise(function(resolve, reject) {
-			Module['wasmInstantiateActionResolve'] = resolve;
-			Module['wasmInstantiateActionReject'] = reject;
-		});
-
-		var mainJsDownload = fetchOrDownloadAndStore(db, Module.locateFile('%CONFIG%.js')).then(function(data) {
-				return addScriptToDom(data).then(function() {
-					addRunDependency('wait-for-compiled-code');
-				});
-			});
-
-		var dataJsDownload = fetchOrDownloadAndStore(db, Module.locateFile('%GAME%.data.js'));
-		var utilityJsDownload = fetchOrDownloadAndStore(db, Module.locateFile('Utility.js')).then(addScriptToDom);
-		var dataDownload = 
-/* // The following code would download and store the .data file as a Blob, which should be more efficient than loading an ArrayBuffer. However that seems to be buggy, so avoid it for now.
-			fetchOrDownloadAndStore(db, Module.locateFile('%GAME%.data')).then(function(dataBlob) {
-				return readBlobToArrayBuffer(dataBlob).then(function(dataArrayBuffer) {
-					Module['preloadedPackages'] = {};
-					Module['preloadedPackages'][Module.locateFile('%GAME%.data')] = dataArrayBuffer;
-					return dataJsDownload.then(addScriptToDom);
-				})
-			});
-*/
-// Instead as a fallback, download as ArrayBuffer. (TODO: Figure out the bugs with the above, and switch to using that one instead)
-			fetchOrDownloadAndStore(db, Module.locateFile('%GAME%.data'), 'arraybuffer').then(function(dataArrayBuffer) {
-				Module['preloadedPackages'] = {};
-				Module['preloadedPackages'][Module.locateFile('%GAME%.data')] = dataArrayBuffer;
-				return dataJsDownload.then(addScriptToDom);
-			});
-
-		const precompileShaders = false; // Currently not enabled.
-
-		if (precompileShaders) {
-			var compileShaders = fetchOrDownloadAndStore(db, Module.locateFile('shaders.json'), 'arraybuffer')
-			.then(function(json) {
-				return compileShadersFromJson(json)
-				.catch(function(error) {
-					taskFinished(TASK_SHADERS, error + '<br>Current renderer: ' + getGpuInfo());
-					throw 'Shader compilation failed';
-				});
-			});
-		} else {
-			var compileShaders = true; // Not precompiling shaders, no-op Promise action.
-		}
-
-		Promise.all([mainCompiledCode, mainJsDownload, dataJsDownload, utilityJsDownload, dataDownload, compiledCodeInstantiateAction, compileShaders]).then(function() {
-			if (!precompileShaders) {
-				Module['precompiledShaders'] = Module['precompiledPrograms'] = Module['preinitializedWebGLContext'] = Module['glIDCounter'] = Module['precompiledUniforms'] = null;
-			}
-			taskProgress(TASK_MAIN);
-			removeRunDependency('wait-for-compiled-code'); // Now we are ready to call main()
-		});
-	};
-
-	openIndexedDB(Module['UE4_indexedDBName'], Module['UE4_indexedDBVersion'] || 1).then(withIndexedDB).catch(function(e) {
-		console.error('Failed to openIndexedDB, proceeding without reading or storing contents to IndexedDB! Error: ');
-		console.error(e);
-		withIndexedDB(null);
-	});
-});
-
-</script>
+<script src="%GAME%.UE4.js"></script>
 </body>
 </html>
diff --git a/Engine/Source/Programs/AutomationTool/HTML5/HTML5Platform.Automation.cs b/Engine/Source/Programs/AutomationTool/HTML5/HTML5Platform.Automation.cs
--- a/Engine/Source/Programs/AutomationTool/HTML5/HTML5Platform.Automation.cs
+++ b/Engine/Source/Programs/AutomationTool/HTML5/HTML5Platform.Automation.cs
@@ -1,4 +1,4 @@
-﻿// Copyright 1998-2018 Epic Games, Inc. All Rights Reserved.
+// Copyright 1998-2018 Epic Games, Inc. All Rights Reserved.
 using System;
 using System.IO;
 using System.Net.Http;
@@ -15,7 +15,6 @@ public class HTML5Platform : Platform
 {
 	// ini configurations
 	static bool Compressed = false;
-	static bool targetWebGL2 = true;
 	static bool enableIndexedDB = false; // experimental for now...
 
 	public HTML5Platform()
@@ -31,21 +30,19 @@ public class HTML5Platform : Platform
 		HTML5SDKInfo.SetupEmscriptenTemp();
 		HTML5SDKInfo.SetUpEmscriptenConfigFile();
 
+		// ----------------------------------------
+		// target output
 		string PackagePath = Path.Combine(Path.GetDirectoryName(Params.RawProjectPath.FullName), "Binaries", "HTML5");
 		if (!Directory.Exists(PackagePath))
 		{
 			Directory.CreateDirectory(PackagePath);
 		}
 
+		// ----------------------------------------
 		// ini configurations
 		var ConfigCache = UnrealBuildTool.ConfigCache.ReadHierarchy(ConfigHierarchyType.Engine, DirectoryReference.FromFile(Params.RawProjectPath), UnrealTargetPlatform.HTML5);
-		bool targetWebGL1 = false; // inverted checked - this will be going away soon...
-		if ( ConfigCache.GetBool("/Script/HTML5PlatformEditor.HTML5TargetSettings", "TargetWebGL1", out targetWebGL1) )
-		{
-			targetWebGL2  = !targetWebGL1;
-		}
 
-		// Debug and Development builds are not uncompressed to:
+		// Debug and Development builds are not compressed to:
 		// - speed up iteration times
 		// - ensure (IndexedDB) data are not cached/used
 		// Shipping builds "can be":
@@ -56,12 +53,13 @@ public class HTML5Platform : Platform
 			ConfigCache.GetBool("/Script/HTML5PlatformEditor.HTML5TargetSettings", "Compressed", out Compressed);
 			ConfigCache.GetBool("/Script/HTML5PlatformEditor.HTML5TargetSettings", "EnableIndexedDB", out enableIndexedDB);
 		}
-		Log("HTML5Platform.Automation: TargetWebGL2 = "     + targetWebGL2    );
 		Log("HTML5Platform.Automation: Compressed = "       + Compressed      );
 		Log("HTML5Platform.Automation: EnableIndexedDB = "  + enableIndexedDB );
 
 		string FinalDataLocation = Path.Combine(PackagePath, Params.ShortProjectName) + ".data";
 
+		// ----------------------------------------
+		// packaging
 		if (HTMLPakAutomation.CanCreateMapPaks(Params))
 		{
 			HTMLPakAutomation PakAutomation = new HTMLPakAutomation(Params, SC);
@@ -108,6 +106,7 @@ public class HTML5Platform : Platform
 			}
 		}
 
+		// ----------------------------------------
 		// copy the "Executable" to the package directory
 		string ProjectGameExeFilename = Params.GetProjectExeForPlatform(UnrealTargetPlatform.HTML5).ToString();
 		string GameBasename = Path.GetFileNameWithoutExtension(ProjectGameExeFilename);
@@ -121,7 +120,7 @@ public class HTML5Platform : Platform
 		string FullPackageGameBasePath = Path.Combine(PackagePath, GameBasename);
 
 		// with extension
-		string GameExe = GameBasename + ".js";
+		string GameExe = GameBasename + ".js"; // emscripten
 		string FullGameExePath = Path.Combine(GameBasepath, GameExe);
 		string FullPackageGameExePath = Path.Combine(PackagePath, GameExe);
 
@@ -143,25 +142,63 @@ public class HTML5Platform : Platform
 		File.SetAttributes(FullPackageGameBasePath + ".wasm", FileAttributes.Normal);
 		File.SetAttributes(FullPackageGameExePath, FileAttributes.Normal);
 
-		// put the HTML file to the package directory
-		string TemplateFile = CombinePaths(CmdEnv.LocalRoot, "Engine/Build/HTML5/GameX.html.template");
+		// ----------------------------------------
+		// generate HTML files to the package directory
+
+		// ini setting
+		string CanvasScaleMode;
+		ConfigCache.GetString("/Script/HTML5PlatformEditor.HTML5TargetSettings", "CanvasScalingMode", out CanvasScaleMode);
+
+		// output base name
 		string OutputFile = Path.Combine(PackagePath,
 				(Params.ClientConfigsToBuild[0].ToString() != "Development" ?
 				 (Params.ShortProjectName + "-HTML5-" + Params.ClientConfigsToBuild[0].ToString()) :
-				 Params.ShortProjectName)) + ".html";
+				 Params.ShortProjectName)); // + ".html";
 
-		string CanvasScaleMode;
-		ConfigCache.GetString("/Script/HTML5PlatformEditor.HTML5TargetSettings", "CanvasScalingMode", out CanvasScaleMode);
+		// custom HTML, JS (if any), and CSS (if any) template files
+		string BuildPath = Path.Combine(Path.GetDirectoryName(Params.RawProjectPath.FullName), "Build", "HTML5");
+		string TemplateFile = CombinePaths(BuildPath, "GameX.html.template");
+		if ( !File.Exists(TemplateFile) )
+		{
+			// fall back to default UE4 template files
+			BuildPath = CombinePaths(CmdEnv.LocalRoot, "Engine/Build/HTML5");
+			TemplateFile = CombinePaths(BuildPath, "GameX.html.template");
+		}
 		GenerateFileFromTemplate(TemplateFile,
-				OutputFile,
+				OutputFile + ".html",
 				Params.ShortProjectName,
 				Params.ClientConfigsToBuild[0].ToString(),
-//				Params.StageCommandline,
 				!Params.IsCodeBasedProject,
 				HTML5SDKInfo.HeapSize(ConfigCache, Params.ClientConfigsToBuild[0].ToString()),
 				CanvasScaleMode
 			);
-
+		TemplateFile = CombinePaths(BuildPath, "GameX.js.template");
+		if ( File.Exists(TemplateFile) )
+		{
+			GenerateFileFromTemplate(TemplateFile,
+					OutputFile + ".UE4.js",
+					Params.ShortProjectName,
+					Params.ClientConfigsToBuild[0].ToString(),
+					!Params.IsCodeBasedProject,
+					HTML5SDKInfo.HeapSize(ConfigCache, Params.ClientConfigsToBuild[0].ToString()),
+					CanvasScaleMode
+				);
+		}
+		TemplateFile = CombinePaths(BuildPath, "GameX.css.template");
+		if ( File.Exists(TemplateFile) )
+		{
+			GenerateFileFromTemplate(TemplateFile,
+					OutputFile + ".css",
+					Params.ShortProjectName,
+					Params.ClientConfigsToBuild[0].ToString(),
+					!Params.IsCodeBasedProject,
+					HTML5SDKInfo.HeapSize(ConfigCache, Params.ClientConfigsToBuild[0].ToString()),
+					CanvasScaleMode
+				);
+		}
+
+		// ----------------------------------------
+		// (development) support files
 		string MacBashTemplateFile = CombinePaths(CmdEnv.LocalRoot, "Engine/Build/HTML5/RunMacHTML5LaunchHelper.command.template");
 		string MacBashOutputFile = Path.Combine(PackagePath, "RunMacHTML5LaunchHelper.command");
 		string MonoPath = CombinePaths(CmdEnv.LocalRoot, "Engine/Build/BatchFiles/Mac/SetupMono.sh");
@@ -180,6 +217,8 @@ public class HTML5Platform : Platform
 		}
 		File.Copy(htaccessTemplate, htaccesspath, true);
 
+		// ----------------------------------------
+		// final copies
 		string JSDir = CombinePaths(CmdEnv.LocalRoot, "Engine/Build/HTML5");
 		string OutDir = PackagePath;
 
@@ -201,23 +240,35 @@ public class HTML5Platform : Platform
 			// Compress all files. These are independent tasks which can be threaded.
 			List<Task> CompressionTasks = new List<Task>();
 
-			// data file
+			// DATA file
 			CompressionTasks.Add(Task.Factory.StartNew(() => CompressFile(FinalDataLocation, FinalDataLocation + "gz")));
 
-			// data file .js driver.
+			// DATA file .js driver (emscripten)
 			CompressionTasks.Add(Task.Factory.StartNew(() => CompressFile(FinalDataLocation + ".js" , FinalDataLocation + ".jsgz")));
 
 			// main game code
 			CompressionTasks.Add(Task.Factory.StartNew(() => CompressFile(FullPackageGameBasePath + ".wasm", FullPackageGameBasePath + ".wasmgz")));
-			// main js.
+			// main js (emscripten)
 			CompressionTasks.Add(Task.Factory.StartNew(() => CompressFile(FullPackageGameExePath, FullPackageGameExePath + "gz")));
 
-			// symbols file.
+			// symbols fil.
 			CompressionTasks.Add(Task.Factory.StartNew(() => CompressFile(FullPackageGameExePath + ".symbols", FullPackageGameExePath + ".symbolsgz")));
 
 			// Utility
 			CompressionTasks.Add(Task.Factory.StartNew(() => CompressFile(OutDir + "/Utility.js", OutDir + "/Utility.jsgz")));
 
+			// UE4 js
+			if ( File.Exists(FullPackageGameBasePath + ".UE4.js") )
+			{
+				CompressionTasks.Add(Task.Factory.StartNew(() => CompressFile(FullPackageGameBasePath + ".UE4.js" , FinalDataLocation + ".UE4.jsgz")));
+			}
+
+			// UE4 css
+			if ( File.Exists(FullPackageGameBasePath + ".css") )
+			{
+				CompressionTasks.Add(Task.Factory.StartNew(() => CompressFile(FullPackageGameBasePath + ".css" , FinalDataLocation + ".cssgz")));
+			}
+
 			Task.WaitAll(CompressionTasks.ToArray());
 		}
 		else
@@ -231,6 +282,8 @@ public class HTML5Platform : Platform
 			File.Delete(FullPackageGameBasePath + ".wasmgz");
 			File.Delete(FullPackageGameExePath + ".symbolsgz");
 			File.Delete(OutDir + "/Utility.jsgz");
+			File.Delete(FullPackageGameBasePath + ".UE4.jsgz");
+			File.Delete(FullPackageGameBasePath + ".cssgz");
 		}
 
 		File.Copy(CombinePaths(CmdEnv.LocalRoot, "Engine/Binaries/DotNET/HTML5LaunchHelper.exe"),CombinePaths(OutDir, "HTML5LaunchHelper.exe"),true);
@@ -273,13 +326,7 @@ public class HTML5Platform : Platform
 		}
 	}
 
-	public override bool RequiresPackageToDeploy
-	{
-		get { return true; }
-	}
-
-
-	protected void GenerateFileFromTemplate(string InTemplateFile, string InOutputFile, string InGameName, string InGameConfiguration, /*string InArguments,*/ bool IsContentOnly, int HeapSize, string CanvasScaleMode)
+	protected void GenerateFileFromTemplate(string InTemplateFile, string InOutputFile, string InGameName, string InGameConfiguration, bool IsContentOnly, int HeapSize, string CanvasScaleMode)
 	{
 		StringBuilder outputContents = new StringBuilder();
 		using (StreamReader reader = new StreamReader(InTemplateFile))
@@ -327,19 +374,6 @@ public class HTML5Platform : Platform
 				{
 					string ArgumentString = "'../../../" + InGameName + "/" + InGameName + ".uproject',";
 					ArgumentString += "'-stdout',"; // suppress double printing to console.log
-
-// this for() and if() combo looks suspect -- but, seems like Params.StageCommandline is not really used anymore...
-//					InArguments = InArguments.Replace ("\"", "");
-//					string[] Arguments = InArguments.Split(' ');
-//					for (int i = 0; i < Arguments.Length - 1; ++i)
-//					{
-//						ArgumentString += "'" + Arguments[i] + "'" + ",' ',";
-//					}
-//					if (Arguments.Length > 0)
-//					{
-//						ArgumentString += "'" + Arguments[Arguments.Length - 1] + "'";
-//					}
-
 					LineStr = LineStr.Replace("%UE4CMDLINE%", ArgumentString);
 				}
 
@@ -357,11 +391,6 @@ public class HTML5Platform : Platform
 					LineStr = LineStr.Replace("%CANVASSCALEMODE%", mode);
 				}
 
-				if (!targetWebGL2 && LineStr.Contains("const explicitlyUseWebGL1"))
-				{
-					LineStr = "const explicitlyUseWebGL1 = true;";
-				}
-
 				outputContents.AppendLine(LineStr);
 			}
 		}
@@ -430,9 +459,8 @@ public class HTML5Platform : Platform
 		}
 	}
 
-	public override void GetFilesToDeployOrStage(ProjectParams Params, DeploymentContext SC)
-	{
-	}
+	// --------------------------------------------------------------------------------
+	// ArchiveCommand.Automation.cs
 
 	public override void GetFilesToArchive(ProjectParams Params, DeploymentContext SC)
 	{
@@ -450,25 +478,29 @@ public class HTML5Platform : Platform
 		{
 			GameBasename += "-HTML5-" + Params.ClientConfigsToBuild[0].ToString();
 		}
-		string GameExe = GameBasename + ".js";
+		string GameExe = GameBasename + ".js"; // emscripten
 
 		// put the HTML file to the package directory
 		string OutputFilename = (Params.ClientConfigsToBuild[0].ToString() != "Development" ?
 				 (Params.ShortProjectName + "-HTML5-" + Params.ClientConfigsToBuild[0].ToString()) :
 				 Params.ShortProjectName) + ".html";
 
-		// data file
+		// DATA file
 		SC.ArchiveFiles(PackagePath, ProjectDataName);
-		// data file js driver
+		// DATA file js driver (emscripten)
 		SC.ArchiveFiles(PackagePath, ProjectDataName + ".js");
-		// main game code
+		// MAIN game code
 		SC.ArchiveFiles(PackagePath, GameBasename + ".wasm");
-		// main js file
+		// MAIN js file (emscripten)
 		SC.ArchiveFiles(PackagePath, GameExe);
 		// symbols file
 		SC.ArchiveFiles(PackagePath, GameExe + ".symbols");
 		// utilities
 		SC.ArchiveFiles(PackagePath, "Utility.js");
+		// UE4 js file
+		SC.ArchiveFiles(PackagePath, GameBasename + ".UE4.js");
+		// UE4 css file
+		SC.ArchiveFiles(PackagePath, GameBasename + ".css");
 		// landing page.
 		SC.ArchiveFiles(PackagePath, OutputFilename);
 
@@ -486,6 +518,8 @@ public class HTML5Platform : Platform
 			SC.ArchiveFiles(PackagePath, GameBasename + ".wasmgz");
 			SC.ArchiveFiles(PackagePath, GameExe + ".symbolsgz");
 			SC.ArchiveFiles(PackagePath, "Utility.jsgz");
+			SC.ArchiveFiles(PackagePath, GameBasename + ".UE4.jsgz");
+			SC.ArchiveFiles(PackagePath, GameBasename + ".cssgz");
 		}
 		else
 		{
@@ -496,6 +530,8 @@ public class HTML5Platform : Platform
 			File.Delete(GameBasename + ".wasmgz");
 			File.Delete(GameExe + ".symbolsgz");
 			File.Delete("Utility.jsgz");
+			File.Delete(GameBasename + ".UE4.jsgz");
+			File.Delete(GameBasename + ".cssgz");
 		}
 
 		if (HTMLPakAutomation.CanCreateMapPaks(Params))
@@ -511,6 +547,9 @@ public class HTML5Platform : Platform
 		UploadToS3(SC, OutputFilename);
 	}
 
+	// --------------------------------------------------------------------------------
+	// RunProjectCommand.Automation.cs
+
 	public override IProcessResult RunClient(ERunOptions ClientRunFlags, string ClientApp, string ClientCmdLine, ProjectParams Params)
 	{
 		// look for browser
@@ -558,10 +597,26 @@ public class HTML5Platform : Platform
 			ProfileDirectory = Path.Combine(ProfileDirectory, "firefox");
 			BrowserCommandline += "  " +  String.Format("-no-remote -profile \\\"{0}\\\"", ProfileDirectory);
 		}
+		else if (LowerBrowserPath.Contains("safari"))
+		{
+			// NOT SUPPORTED: cannot have a separate UE4 profile for safari
+			// -- this "can" be done with a different user (e.g. guest) account...
+			//    (which is not a turn key solution that can be done within UE4)
+			// -- some have tried using symlinks to "mimic" this
+			//    https://discussions.apple.com/thread/3327990
+			// -- but, none of these are fool proof with an existing/running safari instance
+
+			// -- also, "Safari Extensions JS" has been officially deprecated as of Safari 12
+			//    (in favor of using "Safari App Extension")
+			//    https://developer.apple.com/documentation/safariextensions
+
+			// this means, Safari "LaunchOn" (UE4 Editor -> Launch -> Safari) will run with your FULL
+			// Safari profile -- so, all of your "previously opened" tabs will all also be opened...
+		}
 
+		// TODO: test on other platforms to remove this first if() check
 		if (UnrealBuildTool.BuildHostPlatform.Current.Platform == UnrealTargetPlatform.Linux)
 		{
-			// TODO: test on other platforms to remove this if() check
 			if (!Directory.Exists(ProfileDirectory))
 			{
 				Directory.CreateDirectory(ProfileDirectory);
@@ -577,55 +632,59 @@ public class HTML5Platform : Platform
 		return BrowserProcess;
 	}
 
-	public override string GetCookPlatform(bool bDedicatedServer, bool bIsClientOnly)
+	public override List<FileReference> GetExecutableNames(DeploymentContext SC)
 	{
-		return "HTML5";
+		List<FileReference> ExecutableNames = new List<FileReference>();
+		ExecutableNames.Add(FileReference.Combine(SC.ProjectRoot, "Binaries", "HTML5", SC.ShortProjectName));
+		return ExecutableNames;
 	}
 
-	public override string GetCookExtraCommandLine(ProjectParams Params)
-	{
-		return HTMLPakAutomation.CanCreateMapPaks(Params) ? " -GenerateDependenciesForMaps " : "";
-	}
+	// --------------------------------------------------------------------------------
+	// PackageCommand.Automation.cs
 
-	public override PakType RequiresPak(ProjectParams Params)
+	public override bool RequiresPackageToDeploy
 	{
-		return HTMLPakAutomation.CanCreateMapPaks(Params) ? PakType.Never : PakType.Always;
+		get { return true; }
 	}
 
-	public override string GetPlatformPakCommandLine(ProjectParams Params, DeploymentContext SC)
-	{
-		return Compressed ? " -compress" : "";
-	}
+	// --------------------------------------------------------------------------------
+	// CookCommand.Automation.cs
 
-	public override bool DeployLowerCaseFilenames()
+	public override string GetCookPlatform(bool bDedicatedServer, bool bIsClientOnly)
 	{
-		return false;
+		return "HTML5";
 	}
 
-	public override string LocalPathToTargetPath(string LocalPath, string LocalRoot)
+	public override string GetCookExtraCommandLine(ProjectParams Params)
 	{
-		return LocalPath;//.Replace("\\", "/").Replace(LocalRoot, "../../..");
+		return HTMLPakAutomation.CanCreateMapPaks(Params) ? " -GenerateDependenciesForMaps " : "";
 	}
 
-	public override bool IsSupported { get { return true; } }
+	// --------------------------------------------------------------------------------
+	// CopyBuildToStagingDirectory.Automation.cs
 
-	public override List<string> GetDebugFileExtensions()
+	public override void GetFilesToDeployOrStage(ProjectParams Params, DeploymentContext SC)
 	{
-		return new List<string> { ".pdb" };
+		// must implement -- "empty" here
 	}
 
-#region Hooks
-	public override void PreBuildAgenda(UE4Build Build, UE4Build.BuildAgenda Agenda)
+	public override PakType RequiresPak(ProjectParams Params)
 	{
+		return HTMLPakAutomation.CanCreateMapPaks(Params) ? PakType.Never : PakType.Always;
 	}
 
-	public override List<FileReference> GetExecutableNames(DeploymentContext SC)
+	public override string GetPlatformPakCommandLine(ProjectParams Params, DeploymentContext SC)
 	{
-		List<FileReference> ExecutableNames = new List<FileReference>();
-		ExecutableNames.Add(FileReference.Combine(SC.ProjectRoot, "Binaries", "HTML5", SC.ShortProjectName));
-		return ExecutableNames;
+		return Compressed ? " -compress" : "";
 	}
-#endregion
+
+	// --------------------------------------------------------------------------------
+	// AutomationUtils/Platform.cs
+
+	public override bool IsSupported { get { return true; } }
+
+	// --------------------------------------------------------------------------------
+	// --------------------------------------------------------------------------------
 
 #region AMAZON S3
 	public void UploadToS3(DeploymentContext SC, string OutputFilename)
@@ -700,6 +759,8 @@ public class HTML5Platform : Platform
 			{ ".wasmgz", "application/wasm" },
 			{ ".htaccess", "text/plain"},
 			{ ".html", "text/html"},
+			{ ".css", "text/css"},
+			{ ".cssgz", "text/css"},
 			{ ".js", "application/x-javascript" },
 			{ ".jsgz", "application/x-javascript" },
 			{ ".symbols", "text/plain"},
