PR #6162: Fix spelling errors (Contributed by [REDACTED_PERSON])

(ushell-p4-cherrypick of 8395412 by UnrealBot)

diff --git a/Engine/Build/HTML5/Readme.txt b/Engine/Build/HTML5/Readme.txt
--- a/Engine/Build/HTML5/Readme.txt
+++ b/Engine/Build/HTML5/Readme.txt
@@ -7,7 +7,7 @@ The following will help explain what kinds of files these are.
 Development Build
 =================
 
-Packaging "Dewelopment" builds:
+Packaging "Development" builds:
 	* Only uncompressed files are built
 	* Jump down to "Local Testing" to test your project
 
